Media: Fixed Movie does not play after reset to default and adding back in (TTP# 348603)

diff --git a/Engine/Plugins/Media/MediaPlayerEditor/Source/MediaPlayerEditor/Private/Widgets/SMediaPlayerEditorViewer.cpp b/Engine/Plugins/Media/MediaPlayerEditor/Source/MediaPlayerEditor/Private/Widgets/SMediaPlayerEditorViewer.cpp
--- a/Engine/Plugins/Media/MediaPlayerEditor/Source/MediaPlayerEditor/Private/Widgets/SMediaPlayerEditorViewer.cpp
+++ b/Engine/Plugins/Media/MediaPlayerEditor/Source/MediaPlayerEditor/Private/Widgets/SMediaPlayerEditorViewer.cpp
@@ -11,9 +11,19 @@
 
 SMediaPlayerEditorViewer::SMediaPlayerEditorViewer( )
 	: CaptionBuffer(MakeShareable(new FMediaSampleBuffer))
+	, MediaPlayer(nullptr)
 { }
 
 
+SMediaPlayerEditorViewer::~SMediaPlayerEditorViewer()
+{
+	if (MediaPlayer != nullptr)
+	{
+		MediaPlayer->OnMediaChanged().RemoveAll(this);
+	}
+}
+
+
 /* SMediaPlayerEditorPlayer interface
  *****************************************************************************/
 
@@ -242,6 +252,8 @@ void SMediaPlayerEditorViewer::Construct( const FArguments& InArgs, UMediaPlayer
 	Viewport = MakeShareable(new FMediaPlayerEditorViewport());
 	ViewportWidget->SetViewportInterface(Viewport.ToSharedRef());
 
+	MediaPlayer->OnMediaChanged().AddRaw(this, &SMediaPlayerEditorViewer::HandleMediaPlayerMediaChanged);
+
 	ReloadMediaPlayer();
 }
 
@@ -414,6 +426,12 @@ FText SMediaPlayerEditorViewer::HandleElapsedTimeTextBlockText( ) const
 }
 
 
+void SMediaPlayerEditorViewer::HandleMediaPlayerMediaChanged()
+{
+	ReloadMediaPlayer();
+}
+
+
 FText SMediaPlayerEditorViewer::HandleOverlayCaptionText( ) const
 {
 	TSharedPtr<TArray<uint8>, ESPMode::ThreadSafe> CurrentCaption = CaptionBuffer->GetCurrentSample();
diff --git a/Engine/Plugins/Media/MediaPlayerEditor/Source/MediaPlayerEditor/Private/Widgets/SMediaPlayerEditorViewer.h b/Engine/Plugins/Media/MediaPlayerEditor/Source/MediaPlayerEditor/Private/Widgets/SMediaPlayerEditorViewer.h
--- a/Engine/Plugins/Media/MediaPlayerEditor/Source/MediaPlayerEditor/Private/Widgets/SMediaPlayerEditorViewer.h
+++ b/Engine/Plugins/Media/MediaPlayerEditor/Source/MediaPlayerEditor/Private/Widgets/SMediaPlayerEditorViewer.h
@@ -19,6 +19,9 @@ public:
 	/** Default constructor. */
 	SMediaPlayerEditorViewer( );
 
+	/** Destructor. */
+	~SMediaPlayerEditorViewer();
+
 public:
 
 	/**
@@ -60,6 +63,9 @@ private:
 	// Callback for getting the text of the 'Elapsed Time' text block.
 	FText HandleElapsedTimeTextBlockText( ) const;
 
+	/** Callback for when the media player's media has changed. */
+	void HandleMediaPlayerMediaChanged();
+
 	// Callback for getting the text of the movie captions.
 	FText HandleOverlayCaptionText( ) const;
 
