Media: Code cleanup pass

diff --git a/Engine/Plugins/Media/AvfMedia/Source/AvfMedia/Private/Player/AvfMediaPlayer.cpp b/Engine/Plugins/Media/AvfMedia/Source/AvfMedia/Private/Player/AvfMediaPlayer.cpp
--- a/Engine/Plugins/Media/AvfMedia/Source/AvfMedia/Private/Player/AvfMediaPlayer.cpp
+++ b/Engine/Plugins/Media/AvfMedia/Source/AvfMedia/Private/Player/AvfMediaPlayer.cpp
@@ -17,15 +17,10 @@ FAvfMediaPlayer::FAvfMediaPlayer()
 }
 
 
-FAvfMediaPlayer::~FAvfMediaPlayer( )
-{
-}
-
-
 /* IMediaInfo interface
  *****************************************************************************/
 
-FTimespan FAvfMediaPlayer::GetDuration( ) const
+FTimespan FAvfMediaPlayer::GetDuration() const
 {
 	return Duration;
 }
@@ -40,7 +35,7 @@ TRange<float> FAvfMediaPlayer::GetSupportedRates( EMediaPlaybackDirections Direc
 }
 
 
-FString FAvfMediaPlayer::GetUrl( ) const
+FString FAvfMediaPlayer::GetUrl() const
 {
 	return MediaUrl;
 }
@@ -52,13 +47,13 @@ bool FAvfMediaPlayer::SupportsRate( float Rate, bool Unthinned ) const
 }
 
 
-bool FAvfMediaPlayer::SupportsScrubbing( ) const
+bool FAvfMediaPlayer::SupportsScrubbing() const
 {
 	return false;
 }
 
 
-bool FAvfMediaPlayer::SupportsSeeking( ) const
+bool FAvfMediaPlayer::SupportsSeeking() const
 {
 	return false;
 }
@@ -67,7 +62,7 @@ bool FAvfMediaPlayer::SupportsSeeking( ) const
 /* IMediaPlayer interface
  *****************************************************************************/
 
-void FAvfMediaPlayer::Close( )
+void FAvfMediaPlayer::Close()
 {
     CurrentTime = 0;
 	MediaUrl = FString();
@@ -90,49 +85,49 @@ void FAvfMediaPlayer::Close( )
 }
 
 
-const IMediaInfo& FAvfMediaPlayer::GetMediaInfo( ) const 
+const IMediaInfo& FAvfMediaPlayer::GetMediaInfo() const 
 {
 	return *this;
 }
 
 
-float FAvfMediaPlayer::GetRate( ) const
+float FAvfMediaPlayer::GetRate() const
 {
     return CurrentRate;
 }
 
 
-FTimespan FAvfMediaPlayer::GetTime( ) const 
+FTimespan FAvfMediaPlayer::GetTime() const 
 {
     return CurrentTime;
 }
 
 
-const TArray<IMediaTrackRef>& FAvfMediaPlayer::GetTracks( ) const
+const TArray<IMediaTrackRef>& FAvfMediaPlayer::GetTracks() const
 {
 	return Tracks;
 }
 
 
-bool FAvfMediaPlayer::IsLooping( ) const 
+bool FAvfMediaPlayer::IsLooping() const 
 {
     return false;
 }
 
 
-bool FAvfMediaPlayer::IsPaused( ) const
+bool FAvfMediaPlayer::IsPaused() const
 {
     return FMath::IsNearlyZero( [MediaPlayer rate] );
 }
 
 
-bool FAvfMediaPlayer::IsPlaying( ) const
+bool FAvfMediaPlayer::IsPlaying() const
 {
     return (MediaPlayer != nil) && (1.0f == MediaPlayer.rate) && (CurrentTime <= Duration);
 }
 
 
-bool FAvfMediaPlayer::IsReady( ) const
+bool FAvfMediaPlayer::IsReady() const
 {
     // To be ready, we need the AVPlayer setup
     bool bIsReady = (MediaPlayer != nil) && ([MediaPlayer status] == AVPlayerStatusReadyToPlay);
@@ -252,25 +247,25 @@ bool FAvfMediaPlayer::SetRate( float Rate )
 }
 
 
-/* FTickerBase implementation
+/* FTickerBase interface
  *****************************************************************************/
 
-
-bool FAvfMediaPlayer::Tick(float DeltaTime)
+bool FAvfMediaPlayer::Tick( float DeltaTime )
 {
-    CurrentTime = FTimespan::FromSeconds( CMTimeGetSeconds([[MediaPlayer currentItem] currentTime]) );
+    CurrentTime = FTimespan::FromSeconds(CMTimeGetSeconds([[MediaPlayer currentItem] currentTime]));
     
-    if( IsPlaying() )
+    if (IsPlaying())
     {
-        for( IMediaTrackRef& Track : Tracks )
+        for (IMediaTrackRef& Track : Tracks)
         {
             FAvfMediaVideoTrack* AVFTrack = (FAvfMediaVideoTrack*)&Track.Get();
-            if( AVFTrack != nil )
+            
+			if (AVFTrack != nil)
             {
-                AVFTrack->ReadFrameAtTime( [[MediaPlayer currentItem] currentTime] );
-            }
+                AVFTrack->ReadFrameAtTime([[MediaPlayer currentItem] currentTime]);
+			}
         }
     }
     
     return true;
-}
\ No newline at end of file
+}
diff --git a/Engine/Plugins/Media/AvfMedia/Source/AvfMedia/Private/Player/AvfMediaPlayer.h b/Engine/Plugins/Media/AvfMedia/Source/AvfMedia/Private/Player/AvfMediaPlayer.h
--- a/Engine/Plugins/Media/AvfMedia/Source/AvfMedia/Private/Player/AvfMediaPlayer.h
+++ b/Engine/Plugins/Media/AvfMedia/Source/AvfMedia/Private/Player/AvfMediaPlayer.h
@@ -13,15 +13,11 @@ class FAvfMediaPlayer
 {
 public:
 
-	/**
-	 * Default constructor.
-	 */
-	FAvfMediaPlayer( );
+	/** Default constructor. */
+	FAvfMediaPlayer();
 
-	/**
-	 * Destructor.
-	 */
-	~FAvfMediaPlayer();
+	/** Destructor. */
+	~FAvfMediaPlayer() { }
 
 public:
 
@@ -67,39 +63,40 @@ public:
 
 public:
 
-    // FTickerObjectBase
-    virtual bool Tick(float DeltaTime) override;
+    // FTickerObjectBase interface
+
+    virtual bool Tick( float DeltaTime ) override;
 
 private:
-    // The AVFoundation media player
+
+    /** The AVFoundation media player */
     AVPlayer* MediaPlayer;
 
-    // The player item which the media player uses to progress.
+    /** The player item which the media player uses to progress. */
     AVPlayerItem* PlayerItem;
 
-    // The available media tracks.
+    /** The available media tracks. */
     TArray<IMediaTrackRef> Tracks;
 
 private:
-    // Media Info
 
-    // The duration of the media.
+    /** The duration of the media. */
     FTimespan Duration;
 
-    // The current time of the playthrough.
+    /** The current time of the playback. */
     FTimespan CurrentTime;
 
-    // The URL of the currently opened media.
+    /** The URL of the currently opened media. */
     FString MediaUrl;
 
-    // The current playback rate.
+    /** The current playback rate. */
     float CurrentRate;
 
 private:
 
-	// Holds an event delegate that is invoked when media is being unloaded.
+	/** Holds an event delegate that is invoked when media is being unloaded. */
 	FOnMediaClosed ClosedEvent;
 
-	// Holds an event delegate that is invoked when media has been loaded.
+	/** Holds an event delegate that is invoked when media has been loaded. */
 	FOnMediaOpened OpenedEvent;
 };
diff --git a/Engine/Plugins/Media/AvfMedia/Source/AvfMedia/Private/Tracks/AvfMediaAudioTrack.h b/Engine/Plugins/Media/AvfMedia/Source/AvfMedia/Private/Tracks/AvfMediaAudioTrack.h
--- a/Engine/Plugins/Media/AvfMedia/Source/AvfMedia/Private/Tracks/AvfMediaAudioTrack.h
+++ b/Engine/Plugins/Media/AvfMedia/Source/AvfMedia/Private/Tracks/AvfMediaAudioTrack.h
@@ -12,9 +12,7 @@ class FAvfMediaAudioTrack
 {
 public:
 
-	/**
-	 * Creates and initializes a new instance.
-	 */
+	/** Default constructor. */
 	FAvfMediaAudioTrack()
 		: FAvfMediaTrack()
 	{
@@ -26,12 +24,12 @@ public:
 
 	// IMediaTrackAudioDetails interface
 	
-	virtual uint32 GetNumChannels( ) const override
+	virtual uint32 GetNumChannels() const override
 	{
 		return NumChannels;
 	}
 
-	virtual uint32 GetSamplesPerSecond( ) const override
+	virtual uint32 GetSamplesPerSecond() const override
 	{
 		return SamplesPerSecond;
 	}
@@ -40,23 +38,23 @@ public:
 
 	// IMediaTrack interface
 
-	virtual const IMediaTrackAudioDetails& GetAudioDetails( ) const override
+	virtual const IMediaTrackAudioDetails& GetAudioDetails() const override
 	{
 		return *this;
 	}
 
-	virtual const IMediaTrackCaptionDetails& GetCaptionDetails( ) const override
+	virtual const IMediaTrackCaptionDetails& GetCaptionDetails() const override
 	{
 		check(false); // not a caption track
 		return (IMediaTrackCaptionDetails&)*this;
 	}
 
-	virtual EMediaTrackTypes GetType( ) const override
+	virtual EMediaTrackTypes GetType() const override
 	{
 		return EMediaTrackTypes::Audio;
 	}
 
-	virtual const IMediaTrackVideoDetails& GetVideoDetails( ) const override
+	virtual const IMediaTrackVideoDetails& GetVideoDetails() const override
 	{
 		check(false); // not an video track
 		return (IMediaTrackVideoDetails&)*this;
@@ -64,9 +62,9 @@ public:
 
 private:
 
-	// The number of channels.
+	/** The number of channels. */
 	uint32 NumChannels;
 
-	// The number of samples per second.
+	/** The number of samples per second. */
 	uint32 SamplesPerSecond;
 };
diff --git a/Engine/Plugins/Media/AvfMedia/Source/AvfMedia/Private/Tracks/AvfMediaCaptionTrack.h b/Engine/Plugins/Media/AvfMedia/Source/AvfMedia/Private/Tracks/AvfMediaCaptionTrack.h
--- a/Engine/Plugins/Media/AvfMedia/Source/AvfMedia/Private/Tracks/AvfMediaCaptionTrack.h
+++ b/Engine/Plugins/Media/AvfMedia/Source/AvfMedia/Private/Tracks/AvfMediaCaptionTrack.h
@@ -12,9 +12,7 @@ class FAvfMediaCaptionTrack
 {
 public:
 
-	/**
-	 * Creates and initializes a new instance.
-	 */
+	/** Default constructor. */
 	FAvfMediaCaptionTrack()
 		: FAvfMediaTrack()
 	{ }
@@ -27,25 +25,25 @@ public:
 
 	// IMediaTrack interface
 
-	virtual const IMediaTrackAudioDetails& GetAudioDetails( ) const override
+	virtual const IMediaTrackAudioDetails& GetAudioDetails() const override
 	{
 		check(false); // not an audio track
 		return (IMediaTrackAudioDetails&)*this;
 	}
 
-	virtual const IMediaTrackCaptionDetails& GetCaptionDetails( ) const override
+	virtual const IMediaTrackCaptionDetails& GetCaptionDetails() const override
 	{
 		return *this;
 	}
 
-	virtual EMediaTrackTypes GetType( ) const override
+	virtual EMediaTrackTypes GetType() const override
 	{
 		return EMediaTrackTypes::Caption;
 	}
 
-	virtual const IMediaTrackVideoDetails& GetVideoDetails( ) const override
+	virtual const IMediaTrackVideoDetails& GetVideoDetails() const override
 	{
 		check(false); // not a video track
 		return (IMediaTrackVideoDetails&)*this;
 	}
-};
\ No newline at end of file
+};
diff --git a/Engine/Plugins/Media/AvfMedia/Source/AvfMedia/Private/Tracks/AvfMediaVideoTrack.cpp b/Engine/Plugins/Media/AvfMedia/Source/AvfMedia/Private/Tracks/AvfMediaVideoTrack.cpp
--- a/Engine/Plugins/Media/AvfMedia/Source/AvfMedia/Private/Tracks/AvfMediaVideoTrack.cpp
+++ b/Engine/Plugins/Media/AvfMedia/Source/AvfMedia/Private/Tracks/AvfMediaVideoTrack.cpp
@@ -3,7 +3,10 @@
 #include "AvfMediaPrivatePCH.h"
 
 
-FAvfMediaVideoTrack::FAvfMediaVideoTrack(AVAssetTrack* InVideoTrack)
+/* FAvfMediaVideoTrack structors
+ *****************************************************************************/
+
+FAvfMediaVideoTrack::FAvfMediaVideoTrack( AVAssetTrack* InVideoTrack )
     : FAvfMediaTrack()
     , SyncStatus(Default)
     , BitRate(0)
@@ -32,6 +35,17 @@ FAvfMediaVideoTrack::FAvfMediaVideoTrack(AVAssetTrack* InVideoTrack)
 }
 
 
+FAvfMediaVideoTrack::~FAvfMediaVideoTrack()
+{
+    // Destroyed.
+    [AVReader release];
+    [AVVideoOutput release];
+}
+
+
+/* FAvfMediaVideoTrack interface
+ *****************************************************************************/
+
 bool FAvfMediaVideoTrack::SeekToTime( const CMTime& SeekTime )
 {
     bool bSeekComplete = false;
@@ -46,7 +60,7 @@ bool FAvfMediaVideoTrack::SeekToTime( const CMTime& SeekTime )
 }
 
 
-void FAvfMediaVideoTrack::GetVideoDetails(AVAssetTrack* InVideoTrack)
+void FAvfMediaVideoTrack::GetVideoDetails( AVAssetTrack* InVideoTrack )
 {
     // Setup the asset reader to ascertain some further video details.
     NSError* nsError = nil;
@@ -91,19 +105,13 @@ void FAvfMediaVideoTrack::GetVideoDetails(AVAssetTrack* InVideoTrack)
 }
 
 
-FAvfMediaVideoTrack::~FAvfMediaVideoTrack()
-{
-    // Destroyed.
-    [AVReader release];
-    [AVVideoOutput release];
-}
-
 void FAvfMediaVideoTrack::ResetAssetReader()
 {
     [AVReader cancelReading];
     [AVReader release];
 }
 
+
 bool FAvfMediaVideoTrack::IsReady() const
 {
     return AVReader != nil && AVReader.status == AVAssetReaderStatusReading;
@@ -208,4 +216,4 @@ void FAvfMediaVideoTrack::ReadFrameAtTime( const CMTime& AVPlayerTime )
             SyncStatus = Default;
         }
     }
-}
\ No newline at end of file
+}
diff --git a/Engine/Plugins/Media/AvfMedia/Source/AvfMedia/Private/Tracks/AvfMediaVideoTrack.h b/Engine/Plugins/Media/AvfMedia/Source/AvfMedia/Private/Tracks/AvfMediaVideoTrack.h
--- a/Engine/Plugins/Media/AvfMedia/Source/AvfMedia/Private/Tracks/AvfMediaVideoTrack.h
+++ b/Engine/Plugins/Media/AvfMedia/Source/AvfMedia/Private/Tracks/AvfMediaVideoTrack.h
@@ -18,51 +18,51 @@ public:
 	 * @param InVideoTrack The av track information required to build this instance.
 	 */
 	FAvfMediaVideoTrack(AVAssetTrack* InVideoTrack);
+
+	/** Destructor. */
     ~FAvfMediaVideoTrack();
 
-    
 public:
 
 	// IMediaTrackVideoDetails interface
 
-	virtual uint32 GetBitRate( ) const override
+	virtual uint32 GetBitRate() const override
 	{
 		return BitRate;
 	}
 
-	virtual FIntPoint GetDimensions( ) const override
+	virtual FIntPoint GetDimensions() const override
 	{
 		return FIntPoint(Width, Height);
 	}
 
-	virtual float GetFrameRate( ) const override
+	virtual float GetFrameRate() const override
 	{
 		return FrameRate;
 	}
 
-
 public:
 
 	// IMediaTrack interface
 
-	virtual const IMediaTrackAudioDetails& GetAudioDetails( ) const override
+	virtual const IMediaTrackAudioDetails& GetAudioDetails() const override
 	{
 		check(false); // not an audio track
 		return (IMediaTrackAudioDetails&)*this;
 	}
 
-	virtual const IMediaTrackCaptionDetails& GetCaptionDetails( ) const override
+	virtual const IMediaTrackCaptionDetails& GetCaptionDetails() const override
 	{
 		check(false); // not a caption track
 		return (IMediaTrackCaptionDetails&)*this;
 	}
 
-	virtual const IMediaTrackVideoDetails& GetVideoDetails( ) const override
+	virtual const IMediaTrackVideoDetails& GetVideoDetails() const override
 	{
 		return *this;
 	}
 
-    virtual EMediaTrackTypes GetType( ) const override
+    virtual EMediaTrackTypes GetType() const override
     {
         return EMediaTrackTypes::Video;
     }
@@ -78,22 +78,25 @@ public:
     }
 
 public:
+
     /**
      * Read the track information for the frame at the given time. This should be provided by the audio manager to ensure the tracks are in sync.
      *
-     * @param AVPlayerTime - The time the AVPlayer is currently synced to.
+     * @param AVPlayerTime The time the AVPlayer is currently synced to.
      */
     void ReadFrameAtTime( const CMTime& AVPlayerTime );
 
     /**
      * Set the Av Asset reader to a specified time.
      *
-     * @param SeekTime - The time the Track Reader should be set to.
+     * @param SeekTime The time the Track Reader should be set to.
      */
     bool SeekToTime( const CMTime& SeekTime );
 
     /**
      * Is this track ready to begin reading?
+	 *
+	 * @return true if ready, false otherwise.
      */
     bool IsReady() const;
 
@@ -102,24 +105,22 @@ private:
     /**
      * Set up the details interface.
      *
-     * @param InVideoTrack - The track for which we can gather the information.
+     * @param InVideoTrack The track for which we can gather the information.
      */
     void GetVideoDetails(AVAssetTrack* InVideoTrack);
 
-    /**
-     * Reset the AVAssetReader to a different point, which will enable us to regrab frames for seeking and restarting.
-     */
+    /** Reset the AVAssetReader to a different point, which will enable us to regrab frames for seeking and restarting. */
     void ResetAssetReader();
 
 private:
 
-    // Access to the Avfoundations track information.
+    /** Access to the Avfoundations track information. */
     AVAssetReader* AVReader;
 
-    // The AVF track information routed from the avassetreader.
+    /** The AVF track information routed from the avassetreader. */
     AVAssetReaderTrackOutput* AVVideoOutput;
 
-    // The current state of affairs for the reader in relation to core media.
+    /** The current state of affairs for the reader in relation to core media. */
     enum ESyncStatus
     {
         Default,    // Starting state
@@ -128,21 +129,21 @@ private:
         Ready,      // Frame is within tolerance of playback cursor.
     } SyncStatus;
 
-    // Flag which dictates whether we have track information has been successfully loaded.
+    /** Flag which dictates whether we have track information has been successfully loaded. */
     bool bVideoTracksLoaded;
 
-    // The video's average bit rate.
+    /** The video's average bit rate. */
     uint32 BitRate;
 
-    // The video's frame rate.
+    /** The video's frame rate. */
     float FrameRate;
 
-	// The video's height in pixels.
+	/** The video's height in pixels. */
 	uint32 Height;
 
-	// The video's width in pixels.
+	/** The video's width in pixels. */
 	uint32 Width;
 
-    // The collection of output streams we are writing to
+    /** The collection of output streams we are writing to. */
     TArray<IMediaSinkWeakPtr> Sinks;
 };
diff --git a/Engine/Plugins/Media/MediaPlayerEditor/Source/MediaPlayerEditor/Private/AssetTools/MediaPlayerActions.h b/Engine/Plugins/Media/MediaPlayerEditor/Source/MediaPlayerEditor/Private/AssetTools/MediaPlayerActions.h
--- a/Engine/Plugins/Media/MediaPlayerEditor/Source/MediaPlayerEditor/Private/AssetTools/MediaPlayerActions.h
+++ b/Engine/Plugins/Media/MediaPlayerEditor/Source/MediaPlayerEditor/Private/AssetTools/MediaPlayerActions.h
@@ -25,31 +25,31 @@ public:
 
 	// FAssetTypeActions_Base overrides
 
-	virtual bool CanFilter( ) override;
+	virtual bool CanFilter() override;
 	virtual void GetActions( const TArray<UObject*>& InObjects, FMenuBuilder& MenuBuilder ) override;
-	virtual uint32 GetCategories( ) override;
-	virtual FText GetName( ) const override;
-	virtual UClass* GetSupportedClass( ) const override;
-	virtual FColor GetTypeColor( ) const override;
+	virtual uint32 GetCategories() override;
+	virtual FText GetName() const override;
+	virtual UClass* GetSupportedClass() const override;
+	virtual FColor GetTypeColor() const override;
 	virtual bool HasActions( const TArray<UObject*>& InObjects ) const override;
 	virtual void OpenAssetEditor( const TArray<UObject*>& InObjects, TSharedPtr<class IToolkitHost> EditWithinLevelEditor = TSharedPtr<IToolkitHost>() ) override;
 
 private:
 
-	// Callback for selecting CreateMediaSoundWave.
+	/** Callback for selecting CreateMediaSoundWave. */
 	void ExecuteCreateMediaSoundWave( TArray<TWeakObjectPtr<UMediaPlayer>> Objects );
 
-	// Callback for selecting CreateMediaTexture.
+	/** Callback for selecting CreateMediaTexture. */
 	void ExecuteCreateMediaTexture( TArray<TWeakObjectPtr<UMediaPlayer>> Objects );	
 
-	// Callback for executing the 'Play Movie' action.
+	/** Callback for executing the 'Play Movie' action. */
 	void HandlePauseMovieActionExecute( TArray<TWeakObjectPtr<UMediaPlayer>> Objects );
 
-	// Callback for executing the 'Play Movie' action.
+	/** Callback for executing the 'Play Movie' action. */
 	void HandlePlayMovieActionExecute( TArray<TWeakObjectPtr<UMediaPlayer>> Objects );
 
 private:
 
-	// Pointer to the style set to use for toolkits.
+	/** Pointer to the style set to use for toolkits. */
 	TSharedRef<ISlateStyle> Style;
 };
diff --git a/Engine/Plugins/Media/MediaPlayerEditor/Source/MediaPlayerEditor/Private/AssetTools/MediaTextureActions.h b/Engine/Plugins/Media/MediaPlayerEditor/Source/MediaPlayerEditor/Private/AssetTools/MediaTextureActions.h
--- a/Engine/Plugins/Media/MediaPlayerEditor/Source/MediaPlayerEditor/Private/AssetTools/MediaTextureActions.h
+++ b/Engine/Plugins/Media/MediaPlayerEditor/Source/MediaPlayerEditor/Private/AssetTools/MediaTextureActions.h
@@ -16,26 +16,26 @@ public:
 
 	// FAssetTypeActions_Base overrides
 
-	virtual bool CanFilter( ) override;
+	virtual bool CanFilter() override;
 	virtual void GetActions( const TArray<UObject*>& InObjects, FMenuBuilder& MenuBuilder ) override;
-	virtual uint32 GetCategories( ) override;
-	virtual FText GetName( ) const override;
-	virtual UClass* GetSupportedClass( ) const override;
-	virtual FColor GetTypeColor( ) const override;
+	virtual uint32 GetCategories() override;
+	virtual FText GetName() const override;
+	virtual UClass* GetSupportedClass() const override;
+	virtual FColor GetTypeColor() const override;
 	virtual bool HasActions( const TArray<UObject*>& InObjects ) const override;
 	virtual void OpenAssetEditor( const TArray<UObject*>& InObjects, TSharedPtr<class IToolkitHost> EditWithinLevelEditor = TSharedPtr<IToolkitHost>() ) override;
 
 private:
 
-	// Callback for selecting CreateMaterial.
+	/** Callback for selecting CreateMaterial. */
 	void ExecuteCreateMaterial( TArray<TWeakObjectPtr<UTexture>> Objects );
 
-	// Callback for selecting CreateSlateBrush.
+	/** Callback for selecting CreateSlateBrush. */
 	void ExecuteCreateSlateBrush( TArray<TWeakObjectPtr<UTexture>> Objects );
 
-	// Callback for selecting Edit.
+	/** Callback for selecting Edit. */
 	void ExecuteEdit( TArray<TWeakObjectPtr<UTexture>> Objects );
 
-	// Callback for selecting FindMaterials.
+	/** Callback for selecting FindMaterials. */
 	void ExecuteFindMaterials( TWeakObjectPtr<UTexture> Object );
 };
diff --git a/Engine/Plugins/Media/MediaPlayerEditor/Source/MediaPlayerEditor/Private/Factories/MediaPlayerFactory.h b/Engine/Plugins/Media/MediaPlayerEditor/Source/MediaPlayerEditor/Private/Factories/MediaPlayerFactory.h
--- a/Engine/Plugins/Media/MediaPlayerEditor/Source/MediaPlayerEditor/Private/Factories/MediaPlayerFactory.h
+++ b/Engine/Plugins/Media/MediaPlayerEditor/Source/MediaPlayerEditor/Private/Factories/MediaPlayerFactory.h
@@ -14,7 +14,6 @@ class UMediaPlayerFactory
 {
 	GENERATED_UCLASS_BODY()
 
-
 public:
 
 	// UFactory Interface
@@ -24,13 +23,13 @@ public:
 protected:
 
 	/** Reloads the collection of supported media formats. */
-	void ReloadMediaFormats( );
+	void ReloadMediaFormats();
 
 private:
 
-	// Callback for handling the adding of media player factories.
-	void HandleMediaPlayerFactoryAdded( );
+	/** Callback for handling the adding of media player factories. */
+	void HandleMediaPlayerFactoryAdded();
 
-	// Callback for handling the adding of media player factories.
-	void HandleMediaPlayerFactoryRemoved( );
+	/** Callback for handling the adding of media player factories. */
+	void HandleMediaPlayerFactoryRemoved();
 };
diff --git a/Engine/Plugins/Media/MediaPlayerEditor/Source/MediaPlayerEditor/Private/MediaPlayerEditorToolkit.h b/Engine/Plugins/Media/MediaPlayerEditor/Source/MediaPlayerEditor/Private/MediaPlayerEditorToolkit.h
--- a/Engine/Plugins/Media/MediaPlayerEditor/Source/MediaPlayerEditor/Private/MediaPlayerEditorToolkit.h
+++ b/Engine/Plugins/Media/MediaPlayerEditor/Source/MediaPlayerEditor/Private/MediaPlayerEditorToolkit.h
@@ -22,10 +22,8 @@ public:
 	 */
 	FMediaPlayerEditorToolkit( const TSharedRef<ISlateStyle>& InStyle );
 
-	/**
-	 * Destructor.
-	 */
-	virtual ~FMediaPlayerEditorToolkit( );
+	/** Destructor. */
+	virtual ~FMediaPlayerEditorToolkit();
 
 public:
 
@@ -42,7 +40,7 @@ public:
 
 	// FAssetEditorToolkit interface
 
-	virtual FString GetDocumentationLink( ) const override;
+	virtual FString GetDocumentationLink() const override;
 	virtual void RegisterTabSpawners( const TSharedRef<class FTabManager>& TabManager ) override;
 	virtual void UnregisterTabSpawners( const TSharedRef<class FTabManager>& TabManager ) override;
 
@@ -50,10 +48,10 @@ public:
 
 	// IToolkit interface
 
-	virtual FText GetBaseToolkitName( ) const override;
-	virtual FName GetToolkitFName( ) const override;
-	virtual FLinearColor GetWorldCentricTabColorScale( ) const override;
-	virtual FString GetWorldCentricTabPrefix( ) const override;
+	virtual FText GetBaseToolkitName() const override;
+	virtual FName GetToolkitFName() const override;
+	virtual FLinearColor GetWorldCentricTabColorScale() const override;
+	virtual FString GetWorldCentricTabPrefix() const override;
 
 public:
 
@@ -70,70 +68,66 @@ protected:
 
 protected:
 
-	/**
-     * Binds the UI commands to delegates.
-	 */
-	void BindCommands( );
+	/** Binds the UI commands to delegates. */
+	void BindCommands();
 
-	/**
-	 * Builds the toolbar widget for the media player editor.
-	 */
-	void ExtendToolBar( );
+	/** Builds the toolbar widget for the media player editor. */
+	void ExtendToolBar();
 
 	/**
 	 * Gets the playback rate for fast forward.
 	 *
 	 * @return Forward playback rate.
 	 */
-	float GetForwardRate( ) const;
+	float GetForwardRate() const;
 
 	/**
 	 * Gets the playback rate for reverse.
 	 *
 	 * @return Reverse playback rate.
 	 */
-	float GetReverseRate( ) const;
+	float GetReverseRate() const;
 
 private:
 
-	// Callback for determining whether the ForwardMedia action can execute.
-	bool HandleForwardMediaActionCanExecute( ) const;
+	/** Callback for determining whether the ForwardMedia action can execute. */
+	bool HandleForwardMediaActionCanExecute() const;
 
-	// Callback for executing the ForwardMedia action.
-	void HandleForwardMediaActionExecute( );
+	/** Callback for executing the ForwardMedia action. */
+	void HandleForwardMediaActionExecute();
 
-	// Callback for determining whether the PauseMedia action can execute.
-	bool HandlePauseMediaActionCanExecute( ) const;
+	/** Callback for determining whether the PauseMedia action can execute. */
+	bool HandlePauseMediaActionCanExecute() const;
 
-	// Callback for executing the PauseMedia action.
-	void HandlePauseMediaActionExecute( );
+	/** Callback for executing the PauseMedia action. */
+	void HandlePauseMediaActionExecute();
 
-	// Callback for determining whether the PlayMedia action can execute.
-	bool HandlePlayMediaActionCanExecute( ) const;
+	/** Callback for determining whether the PlayMedia action can execute. */
+	bool HandlePlayMediaActionCanExecute() const;
 
-	// Callback for executing the PlayMedia action.
-	void HandlePlayMediaActionExecute( );
+	/** Callback for executing the PlayMedia action. */
+	void HandlePlayMediaActionExecute();
 
-	// Callback for determining whether the ReverseMedia action can execute.
-	bool HandleReverseMediaActionCanExecute( ) const;
+	/** Callback for determining whether the ReverseMedia action can execute. */
+	bool HandleReverseMediaActionCanExecute() const;
 
-	// Callback for executing the ReverseMedia action.
-	void HandleReverseMediaActionExecute( );
+	/** Callback for executing the ReverseMedia action. */
+	void HandleReverseMediaActionExecute();
 
-	// Callback for determining whether the ReverseMedia action can execute.
-	bool HandleRewindMediaActionCanExecute( ) const;
+	/** Callback for determining whether the ReverseMedia action can execute. */
+	bool HandleRewindMediaActionCanExecute() const;
 
-	// Callback for executing the ReverseMedia action.
-	void HandleRewindMediaActionExecute( );
+	/** Callback for executing the ReverseMedia action. */
+	void HandleRewindMediaActionExecute();
 
-	// Callback for spawning the Properties tab.
+	/** Callback for spawning the Properties tab. */
 	TSharedRef<SDockTab> HandleTabManagerSpawnTab( const FSpawnTabArgs& Args, FName TabIdentifier );
 
 private:
 
-	// The media player asset being edited.
+	/** The media player asset being edited. */
 	UMediaPlayer* MediaPlayer;
 
-	// Pointer to the style set to use for toolkits.
+	/** Pointer to the style set to use for toolkits. */
 	TSharedRef<ISlateStyle> Style;
 };
diff --git a/Engine/Plugins/Media/MediaPlayerEditor/Source/MediaPlayerEditor/Private/Models/MediaPlayerEditorViewport.h b/Engine/Plugins/Media/MediaPlayerEditor/Source/MediaPlayerEditor/Private/Models/MediaPlayerEditorViewport.h
--- a/Engine/Plugins/Media/MediaPlayerEditor/Source/MediaPlayerEditor/Private/Models/MediaPlayerEditorViewport.h
+++ b/Engine/Plugins/Media/MediaPlayerEditor/Source/MediaPlayerEditor/Private/Models/MediaPlayerEditorViewport.h
@@ -8,18 +8,14 @@ class FMediaPlayerEditorViewport
 {
 public:
 
-	/**
-	 * Default constructor.
-	 */
-	FMediaPlayerEditorViewport( )
+	/** Default constructor. */
+	FMediaPlayerEditorViewport()
 		: EditorTexture(nullptr)
 		, SlateTexture(nullptr)
 	{ }
 
-	/**
-	 * Destructor.
-	 */
-	~FMediaPlayerEditorViewport( )
+	/** Destructor. */
+	~FMediaPlayerEditorViewport()
 	{
 		ReleaseResources();
 	}
@@ -57,7 +53,7 @@ public:
 
 protected:
 
-	void ReleaseResources( )
+	void ReleaseResources()
 	{
 		if (SlateTexture != nullptr)
 		{
@@ -83,28 +79,28 @@ public:
 
 	// ISlateViewport interface
 
-	virtual FIntPoint GetSize( ) const override
+	virtual FIntPoint GetSize() const override
 	{
 		return (SlateTexture != nullptr)
 			? FIntPoint(SlateTexture->GetWidth(), SlateTexture->GetHeight())
 			: FIntPoint();
 	}
 
-	virtual class FSlateShaderResource* GetViewportRenderTargetTexture( ) const override
+	virtual class FSlateShaderResource* GetViewportRenderTargetTexture() const override
 	{
 		return SlateTexture;
 	}
 
-	virtual bool RequiresVsync( ) const override
+	virtual bool RequiresVsync() const override
 	{
 		return false;
 	}
 
 private:
 
-	// The texture being rendered on this viewport
+	/** The texture being rendered on this viewport. */
 	FMediaPlayerEditorTexture* EditorTexture;
 
-	// Pointer to the Slate texture being rendered.
+	/** Pointer to the Slate texture being rendered. */
 	FSlateTexture2DRHIRef* SlateTexture;
 };
\ No newline at end of file
diff --git a/Engine/Plugins/Media/MediaPlayerEditor/Source/MediaPlayerEditor/Private/Styles/MediaPlayerEditorStyle.h b/Engine/Plugins/Media/MediaPlayerEditor/Source/MediaPlayerEditor/Private/Styles/MediaPlayerEditorStyle.h
--- a/Engine/Plugins/Media/MediaPlayerEditor/Source/MediaPlayerEditor/Private/Styles/MediaPlayerEditorStyle.h
+++ b/Engine/Plugins/Media/MediaPlayerEditor/Source/MediaPlayerEditor/Private/Styles/MediaPlayerEditorStyle.h
@@ -18,10 +18,8 @@ class FMediaPlayerEditorStyle
 {
 public:
 
-	/**
-	 * Default constructor.
-	 */
-	 FMediaPlayerEditorStyle( )
+	/** Default constructor. */
+	 FMediaPlayerEditorStyle()
 		 : FSlateStyleSet("MediaPlayerEditorStyle")
 	 {
 		const FVector2D Icon16x16(16.0f, 16.0f);
@@ -47,10 +45,8 @@ public:
 		FSlateStyleRegistry::RegisterSlateStyle(*this);
 	 }
 
-	 /**
-	  * Destructor.
-	  */
-	 ~FMediaPlayerEditorStyle( )
+	 /** Destructor. */
+	 ~FMediaPlayerEditorStyle()
 	 {
 		FSlateStyleRegistry::UnRegisterSlateStyle(*this);
 	 }
diff --git a/Engine/Plugins/Media/MediaPlayerEditor/Source/MediaPlayerEditor/Private/Widgets/SMediaPlayerEditorDetails.h b/Engine/Plugins/Media/MediaPlayerEditor/Source/MediaPlayerEditor/Private/Widgets/SMediaPlayerEditorDetails.h
--- a/Engine/Plugins/Media/MediaPlayerEditor/Source/MediaPlayerEditor/Private/Widgets/SMediaPlayerEditorDetails.h
+++ b/Engine/Plugins/Media/MediaPlayerEditor/Source/MediaPlayerEditor/Private/Widgets/SMediaPlayerEditorDetails.h
@@ -27,6 +27,6 @@ public:
 
 private:
 
-	// Pointer to the MediaPlayer asset that is being viewed.
+	/** Pointer to the MediaPlayer asset that is being viewed. */
 	UMediaPlayer* MediaPlayer;
 };
diff --git a/Engine/Plugins/Media/MediaPlayerEditor/Source/MediaPlayerEditor/Private/Widgets/SMediaPlayerEditorViewer.h b/Engine/Plugins/Media/MediaPlayerEditor/Source/MediaPlayerEditor/Private/Widgets/SMediaPlayerEditorViewer.h
--- a/Engine/Plugins/Media/MediaPlayerEditor/Source/MediaPlayerEditor/Private/Widgets/SMediaPlayerEditorViewer.h
+++ b/Engine/Plugins/Media/MediaPlayerEditor/Source/MediaPlayerEditor/Private/Widgets/SMediaPlayerEditorViewer.h
@@ -17,7 +17,7 @@ public:
 public:
 
 	/** Default constructor. */
-	SMediaPlayerEditorViewer( );
+	SMediaPlayerEditorViewer();
 
 	/** Destructor. */
 	~SMediaPlayerEditorViewer();
@@ -35,114 +35,112 @@ public:
 
 protected:
 
-	/**
-	 * Reloads the media texture and updates all options.
-	 */
-	void ReloadMediaPlayer( );
+	/** Reloads the media texture and updates all options. */
+	void ReloadMediaPlayer();
 
 private:
 
-	// Callback for generating a widget for a value in the audio track combo box.
+	/** Callback for generating a widget for a value in the audio track combo box. */
 	TSharedRef<SWidget> HandleAudioTrackComboBoxGenerateWidget( IMediaTrackPtr Value ) const;
 
-	// Callback for changing the selected audio track.
+	/** Callback for changing the selected audio track. */
 	void HandleAudioTrackComboBoxSelectionChanged( IMediaTrackPtr Selection, ESelectInfo::Type SelectInfo );
 
-	// Callback for getting the text in the audio track combo box.
-	FText HandleAudioTrackComboBoxText( ) const;
+	/** Callback for getting the text in the audio track combo box. */
+	FText HandleAudioTrackComboBoxText() const;
 
-	// Callback for generating a widget for a value in the caption track combo box.
+	/** Callback for generating a widget for a value in the caption track combo box. */
 	TSharedRef<SWidget> HandleCaptionTrackComboBoxGenerateWidget( IMediaTrackPtr Value ) const;
 
-	// Callback for changing the selected audio track.
+	/** Callback for changing the selected audio track. */
 	void HandleCaptionTrackComboBoxSelectionChanged( IMediaTrackPtr Selection, ESelectInfo::Type SelectInfo );
 
-	// Callback for getting the text in the caption track combo box.
-	FText HandleCaptionTrackComboBoxText( ) const;
+	/** Callback for getting the text in the caption track combo box. */
+	FText HandleCaptionTrackComboBoxText() const;
 
-	// Callback for getting the text of the 'Elapsed Time' text block.
-	FText HandleElapsedTimeTextBlockText( ) const;
+	/** Callback for getting the text of the 'Elapsed Time' text block. */
+	FText HandleElapsedTimeTextBlockText() const;
 
 	/** Callback for when the media player's media has changed. */
 	void HandleMediaPlayerMediaChanged();
 
-	// Callback for getting the text of the movie captions.
-	FText HandleOverlayCaptionText( ) const;
+	/** Callback for getting the text of the movie captions. */
+	FText HandleOverlayCaptionText() const;
 
-	// Callback for getting the text of the playback state overlay.
-	FText HandleOverlayStateText( ) const;
+	/** Callback for getting the text of the playback state overlay. */
+	FText HandleOverlayStateText() const;
 
-	// Callback for getting the maximum value of the playback rate spin box.
-	TOptional<float> HandlePlaybackRateBoxMaxValue( ) const;
+	/** Callback for getting the maximum value of the playback rate spin box. */
+	TOptional<float> HandlePlaybackRateBoxMaxValue() const;
 
-	// Callback for getting the minimum value of the playback rate spin box.
-	TOptional<float> HandlePlaybackRateBoxMinValue( ) const;
+	/** Callback for getting the minimum value of the playback rate spin box. */
+	TOptional<float> HandlePlaybackRateBoxMinValue() const;
 
-	// Callback for getting the current value of the playback rate spin box.
-	TOptional<float> HandlePlaybackRateSpinBoxValue( ) const;
+	/** Callback for getting the current value of the playback rate spin box. */
+	TOptional<float> HandlePlaybackRateSpinBoxValue() const;
 
-	// Callback for committing a new value to the playback rate spin box.
+	/** Callback for committing a new value to the playback rate spin box. */
 	void HandlePlaybackRateBoxValueChanged( float NewValue );
 
-	// Callback for getting the enabled state of the position slider.
-	bool HandlePositionSliderIsEnabled( ) const;
+	/** Callback for getting the enabled state of the position slider. */
+	bool HandlePositionSliderIsEnabled() const;
 
-	// Callback for when the position slider captures the mouse.
-	void HandlePositionSliderMouseCaptureBegin( );
+	/** Callback for when the position slider captures the mouse. */
+	void HandlePositionSliderMouseCaptureBegin();
 
-	// Callback for when the position slider releases the mouse.
-	void HandlePositionSliderMouseCaptureEnd( );
+	/** Callback for when the position slider releases the mouse. */
+	void HandlePositionSliderMouseCaptureEnd();
 
-	// Callback for getting the value of the position slider.
-	float HandlePositionSliderValue( ) const;
+	/** Callback for getting the value of the position slider. */
+	float HandlePositionSliderValue() const;
 
-	// Callback for changing the value of the 'Position' slider.
+	/** Callback for changing the value of the 'Position' slider. */
 	void HandlePositionSliderValueChanged( float NewValue );
 
-	// Callback for getting the text of the 'Remaining Time' text block.
-	FText HandleRemainingTimeTextBlockText( ) const;
+	/** Callback for getting the text of the 'Remaining Time' text block. */
+	FText HandleRemainingTimeTextBlockText() const;
 
-	// Callback for generating a widget for a value in the video track combo box.
+	/** Callback for generating a widget for a value in the video track combo box. */
 	TSharedRef<SWidget> HandleVideoTrackComboBoxGenerateWidget( IMediaTrackPtr Value ) const;
 
-	// Callback for changing the selected video track.
+	/** Callback for changing the selected video track. */
 	void HandleVideoTrackComboBoxSelectionChanged( IMediaTrackPtr Selection, ESelectInfo::Type SelectInfo );
 
-	// Callback for getting the text in the video track combo box.
-	FText HandleVideoTrackComboBoxText( ) const;
+	/** Callback for getting the text in the video track combo box. */
+	FText HandleVideoTrackComboBoxText() const;
 
 private:
 
-	// The collection of available audio tracks.
+	/** The collection of available audio tracks. */
 	TArray<IMediaTrackPtr> AudioTracks;
 
-	// Holds the audio track combo box.
+	/** Holds the audio track combo box. */
 	TSharedPtr<SComboBox<IMediaTrackPtr>> AudioTrackComboBox;
 
-	// Buffers text from the selected caption track.
+	/** Buffers text from the selected caption track. */
 	FMediaSampleBufferRef CaptionBuffer;
 
-	// The collection of available caption tracks.
+	/** The collection of available caption tracks. */
 	TArray<IMediaTrackPtr> CaptionTracks;
 
-	// Holds the caption track combo box.
+	/** Holds the caption track combo box. */
 	TSharedPtr<SComboBox<IMediaTrackPtr>> CaptionTrackComboBox;
 
-	// Pointer to the media player that is being viewed.
+	/** Pointer to the media player that is being viewed. */
 	UMediaPlayer* MediaPlayer;
 
-	// The playback rate prior to scrubbing.
+	/** The playback rate prior to scrubbing. */
 	float PreScrubRate;
 
-	// Holds the scrubber slider.
+	/** Holds the scrubber slider. */
 	TSharedPtr<SSlider> ScrubberSlider;
 
-	// The collection of available video tracks.
+	/** The collection of available video tracks. */
 	TArray<IMediaTrackPtr> VideoTracks;
 
-	// Holds the video track combo box.
+	/** Holds the video track combo box. */
 	TSharedPtr<SComboBox<IMediaTrackPtr>> VideoTrackComboBox;
 
-	// Holds the Slate viewport.
+	/** Holds the Slate viewport. */
 	TSharedPtr<FMediaPlayerEditorViewport> Viewport;
 };
diff --git a/Engine/Plugins/Media/WmfMedia/Source/WmfMedia/Private/Player/WmfMediaPlayer.h b/Engine/Plugins/Media/WmfMedia/Source/WmfMedia/Private/Player/WmfMediaPlayer.h
--- a/Engine/Plugins/Media/WmfMedia/Source/WmfMedia/Private/Player/WmfMediaPlayer.h
+++ b/Engine/Plugins/Media/WmfMedia/Source/WmfMedia/Private/Player/WmfMediaPlayer.h
@@ -14,40 +14,36 @@ class FWmfMediaPlayer
 {
 public:
 
-	/**
-	 * Default constructor.
-	 */
-	FWmfMediaPlayer( );
+	/** Default constructor. */
+	FWmfMediaPlayer();
 
-	/**
-	 * Destructor.
-	 */
-	~FWmfMediaPlayer( );
+	/** Destructor. */
+	~FWmfMediaPlayer();
 
 public:
 
 	// IMediaInfo interface
 
-	virtual FTimespan GetDuration( ) const override;
+	virtual FTimespan GetDuration() const override;
 	virtual TRange<float> GetSupportedRates( EMediaPlaybackDirections Direction, bool Unthinned ) const override;
-	virtual FString GetUrl( ) const override;
+	virtual FString GetUrl() const override;
 	virtual bool SupportsRate( float Rate, bool Unthinned ) const override;
-	virtual bool SupportsScrubbing( ) const override;
-	virtual bool SupportsSeeking( ) const override;
+	virtual bool SupportsScrubbing() const override;
+	virtual bool SupportsSeeking() const override;
 
 public:
 
 	// IMediaPlayer interface
 
-	virtual void Close( ) override;
-	virtual const IMediaInfo& GetMediaInfo( ) const override;
-	virtual float GetRate( ) const override;
-	virtual FTimespan GetTime( ) const override;
-	virtual const TArray<IMediaTrackRef>& GetTracks( ) const override;
-	virtual bool IsLooping( ) const override;
-	virtual bool IsPaused( ) const override;
-	virtual bool IsPlaying( ) const override;
-	virtual bool IsReady( ) const override;
+	virtual void Close() override;
+	virtual const IMediaInfo& GetMediaInfo() const override;
+	virtual float GetRate() const override;
+	virtual FTimespan GetTime() const override;
+	virtual const TArray<IMediaTrackRef>& GetTracks() const override;
+	virtual bool IsLooping() const override;
+	virtual bool IsPaused() const override;
+	virtual bool IsPlaying() const override;
+	virtual bool IsReady() const override;
 	virtual bool Open( const FString& Url ) override;
 	virtual bool Open( const TSharedRef<TArray<uint8>>& Buffer, const FString& OriginalUrl ) override;
 	virtual bool Seek( const FTimespan& Time ) override;
@@ -55,13 +51,13 @@ public:
 	virtual bool SetRate( float Rate ) override;
 
 	DECLARE_DERIVED_EVENT(FWmfMediaPlayer, IMediaPlayer::FOnMediaClosed, FOnMediaClosed);
-	virtual FOnMediaClosed& OnClosed( ) override
+	virtual FOnMediaClosed& OnClosed() override
 	{
 		return ClosedEvent;
 	}
 
 	DECLARE_DERIVED_EVENT(FWmfMediaPlayer, IMediaPlayer::FOnMediaOpened, FOnMediaOpened);
-	virtual FOnMediaOpened& OnOpened( ) override
+	virtual FOnMediaOpened& OnOpened() override
 	{
 		return OpenedEvent;
 	}
diff --git a/Engine/Plugins/Media/WmfMedia/Source/WmfMedia/Private/Tracks/WmfMediaAudioTrack.h b/Engine/Plugins/Media/WmfMedia/Source/WmfMedia/Private/Tracks/WmfMediaAudioTrack.h
--- a/Engine/Plugins/Media/WmfMedia/Source/WmfMedia/Private/Tracks/WmfMediaAudioTrack.h
+++ b/Engine/Plugins/Media/WmfMedia/Source/WmfMedia/Private/Tracks/WmfMediaAudioTrack.h
@@ -31,12 +31,12 @@ public:
 
 	// IMediaTrackAudioDetails interface
 	
-	virtual uint32 GetNumChannels( ) const override
+	virtual uint32 GetNumChannels() const override
 	{
 		return NumChannels;
 	}
 
-	virtual uint32 GetSamplesPerSecond( ) const override
+	virtual uint32 GetSamplesPerSecond() const override
 	{
 		return SamplesPerSecond;
 	}
@@ -45,23 +45,23 @@ public:
 
 	// IMediaTrack interface
 
-	virtual const IMediaTrackAudioDetails& GetAudioDetails( ) const override
+	virtual const IMediaTrackAudioDetails& GetAudioDetails() const override
 	{
 		return *this;
 	}
 
-	virtual const IMediaTrackCaptionDetails& GetCaptionDetails( ) const override
+	virtual const IMediaTrackCaptionDetails& GetCaptionDetails() const override
 	{
 		check(false); // not a caption track
 		return (IMediaTrackCaptionDetails&)*this;
 	}
 
-	virtual EMediaTrackTypes GetType( ) const override
+	virtual EMediaTrackTypes GetType() const override
 	{
 		return EMediaTrackTypes::Audio;
 	}
 
-	virtual const IMediaTrackVideoDetails& GetVideoDetails( ) const override
+	virtual const IMediaTrackVideoDetails& GetVideoDetails() const override
 	{
 		check(false); // not an video track
 		return (IMediaTrackVideoDetails&)*this;
@@ -69,10 +69,10 @@ public:
 
 private:
 
-	// The number of channels.
+	/** The number of channels. */
 	uint32 NumChannels;
 
-	// The number of samples per second.
+	/** The number of samples per second. */
 	uint32 SamplesPerSecond;
 };
 
diff --git a/Engine/Plugins/Media/WmfMedia/Source/WmfMedia/Private/Tracks/WmfMediaVideoTrack.h b/Engine/Plugins/Media/WmfMedia/Source/WmfMedia/Private/Tracks/WmfMediaVideoTrack.h
--- a/Engine/Plugins/Media/WmfMedia/Source/WmfMedia/Private/Tracks/WmfMediaVideoTrack.h
+++ b/Engine/Plugins/Media/WmfMedia/Source/WmfMedia/Private/Tracks/WmfMediaVideoTrack.h
@@ -44,17 +44,17 @@ public:
 
 	// IMediaTrackVideoDetails interface
 
-	virtual uint32 GetBitRate( ) const override
+	virtual uint32 GetBitRate() const override
 	{
 		return BitRate;
 	}
 
-	virtual FIntPoint GetDimensions( ) const override
+	virtual FIntPoint GetDimensions() const override
 	{
 		return FIntPoint(Width, Height);
 	}
 
-	virtual float GetFrameRate( ) const override
+	virtual float GetFrameRate() const override
 	{
 		return FrameRate;
 	}
@@ -63,40 +63,40 @@ public:
 
 	// IMediaTrack interface
 
-	virtual const IMediaTrackAudioDetails& GetAudioDetails( ) const override
+	virtual const IMediaTrackAudioDetails& GetAudioDetails() const override
 	{
 		check(false); // not an audio track
 		return (IMediaTrackAudioDetails&)*this;
 	}
 
-	virtual const IMediaTrackCaptionDetails& GetCaptionDetails( ) const override
+	virtual const IMediaTrackCaptionDetails& GetCaptionDetails() const override
 	{
 		check(false); // not a caption track
 		return (IMediaTrackCaptionDetails&)*this;
 	}
 
-	virtual EMediaTrackTypes GetType( ) const override
+	virtual EMediaTrackTypes GetType() const override
 	{
 		return EMediaTrackTypes::Video;
 	}
 
-	virtual const IMediaTrackVideoDetails& GetVideoDetails( ) const override
+	virtual const IMediaTrackVideoDetails& GetVideoDetails() const override
 	{
 		return *this;
 	}
 
 private:
 
-	// The video's average bit rate.
+	/** The video's average bit rate. */
 	uint32 BitRate;
 
-	// The video's frame rate.
+	/** The video's frame rate. */
 	float FrameRate;
 
-	// The video's height in pixels.
+	/** The video's height in pixels. */
 	UINT32 Height;
 
-	// The video's width in pixels.
+	/** The video's width in pixels. */
 	UINT32 Width;
 };
 
diff --git a/Engine/Plugins/Media/WmfMedia/Source/WmfMedia/Private/Wmf/WmfMediaByteStream.h b/Engine/Plugins/Media/WmfMedia/Source/WmfMedia/Private/Wmf/WmfMediaByteStream.h
--- a/Engine/Plugins/Media/WmfMedia/Source/WmfMedia/Private/Wmf/WmfMediaByteStream.h
+++ b/Engine/Plugins/Media/WmfMedia/Source/WmfMedia/Private/Wmf/WmfMediaByteStream.h
@@ -22,17 +22,17 @@ public:
 	FWmfMediaByteStream( const TSharedRef<TArray<uint8>>& InBuffer );
 
 	/** Virtual destructor. */
-	virtual ~FWmfMediaByteStream( ) { }
+	virtual ~FWmfMediaByteStream() { }
 
 public:
 
 	// IMFAsyncCallback interface
 
-	STDMETHODIMP_(ULONG) AddRef( );
+	STDMETHODIMP_(ULONG) AddRef();
 	STDMETHODIMP GetParameters( unsigned long*, unsigned long*);
 	STDMETHODIMP Invoke( IMFAsyncResult* AsyncResult );
 	STDMETHODIMP QueryInterface( REFIID RefID, void** Object );
-	STDMETHODIMP_(ULONG) Release( );
+	STDMETHODIMP_(ULONG) Release();
 
 public:
 
@@ -40,10 +40,10 @@ public:
 
 	STDMETHODIMP BeginRead( BYTE* pb, ULONG cb, IMFAsyncCallback* pCallback, IUnknown* punkState );
 	STDMETHODIMP BeginWrite( const BYTE* pb, ULONG cb, IMFAsyncCallback* pCallback, IUnknown* punkState );
-	STDMETHODIMP Close( );
+	STDMETHODIMP Close();
 	STDMETHODIMP EndRead( IMFAsyncResult* pResult, ULONG* pcbRead );
 	STDMETHODIMP EndWrite( IMFAsyncResult* pResult, ULONG* pcbWritten );
-	STDMETHODIMP Flush( );
+	STDMETHODIMP Flush();
 	STDMETHODIMP GetCapabilities( DWORD* pdwCapabilities );
 	STDMETHODIMP GetCurrentPosition( QWORD* pqwPosition );
 	STDMETHODIMP GetLength( QWORD* pqwLength );
@@ -56,19 +56,19 @@ public:
 
 private:
 
-	// Whether the stream is currently being read asynchronously.
+	/** Whether the stream is currently being read asynchronously. */
 	bool AsyncReadInProgress;
 
-	// Holds the buffer to stream from.
+	/** Holds the buffer to stream from. */
 	TSharedRef<TArray<uint8>> Buffer;
 
-	// Critical section for locking access to this class.
+	/** Critical section for locking access to this class. */
 	mutable FCriticalSection CriticalSection;
 
-	// The current read position.
+	/** The current read position. */
 	uint64 Position;
 
-	// Holds a reference counter for this instance.
+	/** Holds a reference counter for this instance. */
 	int32 RefCount;
 };
 
diff --git a/Engine/Plugins/Media/WmfMedia/Source/WmfMedia/Private/Wmf/WmfMediaReadState.h b/Engine/Plugins/Media/WmfMedia/Source/WmfMedia/Private/Wmf/WmfMediaReadState.h
--- a/Engine/Plugins/Media/WmfMedia/Source/WmfMedia/Private/Wmf/WmfMediaReadState.h
+++ b/Engine/Plugins/Media/WmfMedia/Source/WmfMedia/Private/Wmf/WmfMediaReadState.h
@@ -31,16 +31,18 @@ public:
 	{ }
 
 	/** Virtual destructor. */
-	virtual ~FWmfMediaReadState( ) { }
+	virtual ~FWmfMediaReadState() { }
 
 public:
 
 	/**
-	 * Sets the number of bytes read.
+	 * Adds the specified number of bytes read.
+	 *
+	 * @param BytesToAdd The number of bytes to add.
 	 */
-	void AddBytesRead( uint64 InBytesRead )
+	void AddBytesRead( uint64 BytesToAdd )
 	{
-		BytesRead += InBytesRead;
+		BytesRead += BytesToAdd;
 	}
 
 	/**
@@ -48,7 +50,7 @@ public:
 	 *
 	 * @return The buffer.
 	 */
-	BYTE* GetReadBuffer( ) const
+	BYTE* GetReadBuffer() const
 	{
 		return ReadBuffer;
 	}
@@ -58,7 +60,7 @@ public:
 	 *
 	 * @return The size of the buffer (in bytes).
 	 */
-	ULONG GetReadBufferSize( ) const
+	ULONG GetReadBufferSize() const
 	{
 		return ReadBufferSize;
 	}
@@ -68,7 +70,7 @@ public:
 	 *
 	 * @return Bytes read count.
 	 */
-	uint64 GetBytesRead( ) const
+	uint64 GetBytesRead() const
 	{
 		return BytesRead;
 	}
@@ -88,12 +90,12 @@ public:
 		return QISearch( this, QITab, RefID, Object );
 	}
 
-	STDMETHODIMP_(ULONG) AddRef( )
+	STDMETHODIMP_(ULONG) AddRef()
 	{
 		return FPlatformAtomics::InterlockedIncrement(&RefCount);
 	}
 
-	STDMETHODIMP_( ULONG ) Release( )
+	STDMETHODIMP_( ULONG ) Release()
 	{
 		int32 CurrentRefCount = FPlatformAtomics::InterlockedDecrement(&RefCount);
 	
@@ -107,16 +109,16 @@ public:
 
 private:
 
-	// Number of bytes read.
+	/** Number of bytes read. */
 	uint64 BytesRead;
 
-	// The buffer that receives the read data.
+	/** The buffer that receives the read data. */
 	BYTE* ReadBuffer;
 
-	// The size of the read buffer.
+	/** The size of the read buffer. */
 	ULONG ReadBufferSize;
 
-	// Holds a reference counter for this instance.
+	/** Holds a reference counter for this instance. */
 	int32 RefCount;
 };
 
diff --git a/Engine/Plugins/Media/WmfMedia/Source/WmfMedia/Private/Wmf/WmfMediaSampler.h b/Engine/Plugins/Media/WmfMedia/Source/WmfMedia/Private/Wmf/WmfMediaSampler.h
--- a/Engine/Plugins/Media/WmfMedia/Source/WmfMedia/Private/Wmf/WmfMediaSampler.h
+++ b/Engine/Plugins/Media/WmfMedia/Source/WmfMedia/Private/Wmf/WmfMediaSampler.h
@@ -17,12 +17,12 @@ public:
 	DECLARE_DELEGATE(CommandDelegate)
 
 	/** Default constructor. */
-	FWmfMediaSampler( )
+	FWmfMediaSampler()
 		: RefCount(1)
 	{ }
 
 	/** Virtual destructor. */
-	virtual ~FWmfMediaSampler( ) { }
+	virtual ~FWmfMediaSampler() { }
 
 public:
 
@@ -54,7 +54,7 @@ public:
 
 	// IMFSampleGrabberSinkCallback interface
 
-	STDMETHODIMP_(ULONG) AddRef( )
+	STDMETHODIMP_(ULONG) AddRef()
 	{
 		return FPlatformAtomics::InterlockedIncrement(&RefCount);
 	}
@@ -113,7 +113,7 @@ public:
 		return S_OK;
 	}
 
-	STDMETHODIMP OnShutdown( )
+	STDMETHODIMP OnShutdown()
 	{
 		return S_OK;
 	}
@@ -130,7 +130,7 @@ public:
 		return QISearch(this, QITab, RefID, Object);
 	}
 
-	STDMETHODIMP_(ULONG) Release( )
+	STDMETHODIMP_(ULONG) Release()
 	{
 		int32 CurrentRefCount = FPlatformAtomics::InterlockedDecrement(&RefCount);
 
@@ -156,13 +156,13 @@ private:
 
 private:
 
-	// Holds the router command queue.
+	/** Holds the router command queue. */
 	TQueue<CommandDelegate, EQueueMode::Mpsc> Commands;
 
-	// The collection of registered media sinks.
+	/** The collection of registered media sinks. */
 	TArray<IMediaSinkWeakPtr> Sinks;
 
-	// Holds a reference counter for this instance.
+	/** Holds a reference counter for this instance. */
 	int32 RefCount;
 };
 
diff --git a/Engine/Plugins/Media/WmfMedia/Source/WmfMedia/Private/Wmf/WmfMediaSession.h b/Engine/Plugins/Media/WmfMedia/Source/WmfMedia/Private/Wmf/WmfMediaSession.h
--- a/Engine/Plugins/Media/WmfMedia/Source/WmfMedia/Private/Wmf/WmfMediaSession.h
+++ b/Engine/Plugins/Media/WmfMedia/Source/WmfMedia/Private/Wmf/WmfMediaSession.h
@@ -53,7 +53,7 @@ public:
 	FWmfMediaSession( const FTimespan& InDuration, const TComPtr<IMFTopology>& InTopology );
 
 	/** Virtual destructor .*/
-	virtual ~FWmfMediaSession( ) { }
+	virtual ~FWmfMediaSession() { }
 
 public:
 
@@ -62,7 +62,7 @@ public:
 	 *
 	 * @return Capabilities bit mask.
 	 */
-	DWORD GetCapabilities( ) const
+	DWORD GetCapabilities() const
 	{
 		return Capabilities;
 	}
@@ -73,7 +73,7 @@ public:
 	 * @return Error result code.
 	 * @see GetSessionState
 	 */
-	HRESULT GetLastError( ) const
+	HRESULT GetLastError() const
 	{
 		return LastError;
 	}
@@ -83,14 +83,14 @@ public:
 	 *
 	 * @return Presentation clock time, or FTimespan::MinValue if no clock is available.
 	 */
-	FTimespan GetPosition( ) const;
+	FTimespan GetPosition() const;
 
 	/**
 	 * Gets the current playback rate.
 	 *
 	 * @return Playback rate.
 	 */
-	float GetRate( ) const
+	float GetRate() const
 	{
 		return CurrentRate;
 	}
@@ -101,7 +101,7 @@ public:
 	 * @return Media session state.
 	 * @see GetLastError
 	 */
-	EMediaStates GetState( ) const
+	EMediaStates GetState() const
 	{
 		return CurrentState;
 	}
@@ -120,7 +120,7 @@ public:
 	 *
 	 * @return true if looping, false otherwise.
 	 */
-	bool IsLooping( ) const
+	bool IsLooping() const
 	{
 		return Looping;
 	}
@@ -173,7 +173,7 @@ public:
 	 *
 	 * @return true if scrubbing is supported, false otherwise.
 	 */
-	bool SupportsScrubbing( ) const
+	bool SupportsScrubbing() const
 	{
 		return CanScrub;
 	}
@@ -182,11 +182,11 @@ public:
 
 	// IMFAsyncCallback interface
 
-	STDMETHODIMP_(ULONG) AddRef( );
+	STDMETHODIMP_(ULONG) AddRef();
 	STDMETHODIMP GetParameters( unsigned long*, unsigned long*);
 	STDMETHODIMP Invoke( IMFAsyncResult* AsyncResult );
 	STDMETHODIMP QueryInterface( REFIID RefID, void** Object );
-	STDMETHODIMP_(ULONG) Release( );
+	STDMETHODIMP_(ULONG) Release();
 
 protected:
 
@@ -195,74 +195,76 @@ protected:
 	 *
 	 * @return true if the state is being changed, false otherwise.
 	 */
-	bool ChangeState( );
+	bool ChangeState();
 
 	/**
 	 * Gets the playback position as known by the internal WMF clock.
 	 *
 	 * @return The internal playback position.
 	 */
-	FTimespan GetInternalPosition( ) const;
+	FTimespan GetInternalPosition() const;
 
 	/**
-	 * Updates the current state.
+	 * Updates the state machine after a state transition was completed.
+	 *
+	 * @param CompletedState The state that was completed.
 	 */
 	void UpdateState( EMediaStates CompletedState );
 
 private:
 
-	// Caches whether the media supports scrubbing.
+	/** Caches whether the media supports scrubbing. */
 	bool CanScrub;
 
-	// The session's capabilities.
+	/** The session's capabilities. */
 	DWORD Capabilities;
 
-	// Whether a new state change has been requested.
+	/** Whether a new state change has been requested. */
 	bool ChangeRequested;
 
-	// Critical section for locking access to this object.
+	/** Critical section for locking access to this object. */
 	mutable FCriticalSection CriticalSection;
 
-	// The current playback rate.
+	/** The current playback rate. */
 	float CurrentRate;
 
-	// The current playback state.
+	/** The current playback state. */
 	EMediaStates CurrentState;
 
-	// The duration of the media.
+	/** The duration of the media. */
 	FTimespan Duration;
 
-	// The media session that handles all playback.
+	/** The media session that handles all playback. */
 	TComPtr<IMFMediaSession> MediaSession;
 
-	// Caches the last error result code (0 = no error).
+	/** Caches the last error result code (0 = no error). */
 	HRESULT LastError;
 
-	// Whether playback should loop to the beginning.
+	/** Whether playback should loop to the beginning. */
 	bool Looping;
 
-	// The media session's clock.
+	/** The media session's clock. */
 	TComPtr<IMFPresentationClock> PresentationClock;
 
-	// Optional interface for controlling playback rates.
+	/** Optional interface for controlling playback rates. */
 	TComPtr<IMFRateControl> RateControl;
 
-	// Optional interface for querying supported playback rates.
+	/** Optional interface for querying supported playback rates. */
 	TComPtr<IMFRateSupport> RateSupport;
 
-	// Holds a reference counter for this instance.
+	/** Holds a reference counter for this instance. */
 	int32 RefCount;
 
-	// The requested playback position.
+	/** The requested playback position. */
 	FTimespan RequestedPosition;
 
-	// The requested playback rate.
+	/** The requested playback rate. */
 	float RequestedRate;
 
-	// The requested playback state.
+	/** The requested playback state. */
 	EMediaStates RequestedState;
 
-	// Whether a state change is currently in progress.
+	/** Whether a state change is currently in progress. */
 	bool StateChangePending;
 };
 
diff --git a/Engine/Plugins/Media/WmfMedia/Source/WmfMedia/Private/WmfMediaModule.cpp b/Engine/Plugins/Media/WmfMedia/Source/WmfMedia/Private/WmfMediaModule.cpp
--- a/Engine/Plugins/Media/WmfMedia/Source/WmfMedia/Private/WmfMediaModule.cpp
+++ b/Engine/Plugins/Media/WmfMedia/Source/WmfMedia/Private/WmfMediaModule.cpp
@@ -23,10 +23,8 @@ class FWmfMediaModule
 {
 public:
 
-	/**
-	 * Default constructor.
-	 */
-	FWmfMediaModule( )
+	/** Default constructor. */
+	FWmfMediaModule()
 		: Initialized(false)
 	{ }
 
@@ -34,7 +32,7 @@ public:
 
 	// IModuleInterface interface
 
-	virtual void StartupModule( ) override
+	virtual void StartupModule() override
 	{
 		// load required libraries
 		IMediaModule* MediaModule = FModuleManager::LoadModulePtr<IMediaModule>("Media");
@@ -85,7 +83,7 @@ public:
 		Initialized = true;
 	}
 
-	virtual void ShutdownModule( ) override
+	virtual void ShutdownModule() override
 	{
 		if (!Initialized)
 		{
@@ -110,7 +108,7 @@ public:
 
 	// IMediaPlayerFactory interface
 
-	virtual TSharedPtr<IMediaPlayer> CreatePlayer( ) override
+	virtual TSharedPtr<IMediaPlayer> CreatePlayer() override
 	{
 		if (Initialized)
 		{
@@ -120,7 +118,7 @@ public:
 		return nullptr;
 	}
 
-	virtual const FMediaFormats& GetSupportedFormats( ) const override
+	virtual const FMediaFormats& GetSupportedFormats() const override
 	{
 		return SupportedFormats;
 	}
@@ -132,7 +130,7 @@ protected:
 	 *
 	 * @return true on success, false otherwise.
 	 */
-	bool LoadRequiredLibraries( )
+	bool LoadRequiredLibraries()
 	{
 		if (LoadLibraryW(TEXT("shlwapi.dll")) == nullptr)
 		{
@@ -167,10 +165,10 @@ protected:
 
 private:
 
-	// Whether the module has been initialized.
+	/** Whether the module has been initialized. */
 	bool Initialized;
 
-	// The collection of supported media formats.
+	/** The collection of supported media formats. */
 	FMediaFormats SupportedFormats;
 };
 
