MediaPlayerEditor: Added notification and viewport overlay for when no media is selected

diff --git a/Engine/Plugins/Media/MediaPlayerEditor/Source/MediaPlayerEditor/Private/Widgets/SMediaPlayerEditorViewer.cpp b/Engine/Plugins/Media/MediaPlayerEditor/Source/MediaPlayerEditor/Private/Widgets/SMediaPlayerEditorViewer.cpp
--- a/Engine/Plugins/Media/MediaPlayerEditor/Source/MediaPlayerEditor/Private/Widgets/SMediaPlayerEditorViewer.cpp
+++ b/Engine/Plugins/Media/MediaPlayerEditor/Source/MediaPlayerEditor/Private/Widgets/SMediaPlayerEditorViewer.cpp
@@ -4,7 +4,7 @@
 #include "SNumericEntryBox.h"
 
 
-#define LOCTEXT_NAMESPACE "SMediaPlayerEditorPlayer"
+#define LOCTEXT_NAMESPACE "SMediaPlayerEditorViewer"
 
 
 /* SMediaPlayerEditorPlayer structors
@@ -31,8 +31,7 @@ SMediaPlayerEditorViewer::~SMediaPlayerEditorViewer()
 void SMediaPlayerEditorViewer::Construct( const FArguments& InArgs, UMediaPlayer* InMediaPlayer, const TSharedRef<ISlateStyle>& InStyle )
 {
 	MediaPlayer = InMediaPlayer;
-
-	TSharedPtr<SViewport> ViewportWidget;
+	Viewport = MakeShareable(new FMediaPlayerEditorViewport());
 
 	ChildSlot
 	[
@@ -42,130 +41,160 @@ void SMediaPlayerEditorViewer::Construct( const FArguments& InArgs, UMediaPlayer
 			.AutoHeight()
 			.Padding(4.0f, 0.0f)
 			[
-				SNew(SHorizontalBox)
+				SNew(SOverlay)
 
-				+ SHorizontalBox::Slot()
-					.AutoWidth()
+				+ SOverlay::Slot()
+					.VAlign(VAlign_Center)
 					[
-						// video track selector
 						SNew(SHorizontalBox)
+							.Visibility(this, &SMediaPlayerEditorViewer::HandleNoMediaSelectedTextVisibility)
 
 						+ SHorizontalBox::Slot()
 							.AutoWidth()
 							.VAlign(VAlign_Center)
 							[
-								SNew(STextBlock)
-									.Text(LOCTEXT("VideoTrackLabel", "Video Track:"))
+								SNew(SImage)
+									.Image(FCoreStyle::Get().GetBrush("Icons.Error"))
 							]
 
 						+ SHorizontalBox::Slot()
 							.FillWidth(1.0f)
-							.Padding(4.0f, 0.0f)
+							.Padding(4.0f, 0.0f, 0.0f, 0.0f)
 							.VAlign(VAlign_Center)
 							[
-								SAssignNew(VideoTrackComboBox, SComboBox<IMediaTrackPtr>)
-									.OnGenerateWidget(this, &SMediaPlayerEditorViewer::HandleVideoTrackComboBoxGenerateWidget)
-									.OnSelectionChanged(this, &SMediaPlayerEditorViewer::HandleVideoTrackComboBoxSelectionChanged)
-									.OptionsSource(&VideoTracks)
-									.ToolTipText(LOCTEXT("CaptionTrackToolTip", "Select the video track to be played"))
-									[
-										SNew(STextBlock)
-											.Text(this, &SMediaPlayerEditorViewer::HandleVideoTrackComboBoxText)
-									]
+								SNew(STextBlock)
+									.Text(LOCTEXT("NoMediaSelectedText", "Please pick a media source in the Details panel!"))
 							]
 					]
 
-				+ SHorizontalBox::Slot()
-					.AutoWidth()
-					.Padding(16.0f, 0.0f, 0.0f, 0.0f)
+				+ SOverlay::Slot()
 					[
-						// audio track selector
 						SNew(SHorizontalBox)
+							.Visibility(this, &SMediaPlayerEditorViewer::HandleTrackSelectionBoxVisibility)
 
 						+ SHorizontalBox::Slot()
 							.AutoWidth()
-							.VAlign(VAlign_Center)
 							[
-								SNew(STextBlock)
-									.Text(LOCTEXT("AudioTrackLabel", "Audio Track:"))
-							]
+								// video track selector
+								SNew(SHorizontalBox)
 
-						+ SHorizontalBox::Slot()
-							.FillWidth(1.0f)
-							.Padding(4.0f, 0.0f, 0.0f, 0.0f)
-							.VAlign(VAlign_Center)
-							[
-								SAssignNew(AudioTrackComboBox, SComboBox<IMediaTrackPtr>)
-									.OnGenerateWidget(this, &SMediaPlayerEditorViewer::HandleAudioTrackComboBoxGenerateWidget)
-									.OnSelectionChanged(this, &SMediaPlayerEditorViewer::HandleAudioTrackComboBoxSelectionChanged)
-									.OptionsSource(&AudioTracks)
-									.ToolTipText(LOCTEXT("CaptionTrackToolTip", "Select the audio track to be played"))
+								+ SHorizontalBox::Slot()
+									.AutoWidth()
+									.VAlign(VAlign_Center)
 									[
 										SNew(STextBlock)
-											.Text(this, &SMediaPlayerEditorViewer::HandleAudioTrackComboBoxText)
+											.Text(LOCTEXT("VideoTrackLabel", "Video Track:"))
 									]
-							]
-					]
 
-				+ SHorizontalBox::Slot()
-					.AutoWidth()
-					.Padding(16.0f, 0.0f, 0.0f, 0.0f)
-					[
-						// caption track selector
-						SNew(SHorizontalBox)
+								+ SHorizontalBox::Slot()
+									.FillWidth(1.0f)
+									.Padding(4.0f, 0.0f)
+									.VAlign(VAlign_Center)
+									[
+										SAssignNew(VideoTrackComboBox, SComboBox<IMediaTrackPtr>)
+											.OnGenerateWidget(this, &SMediaPlayerEditorViewer::HandleVideoTrackComboBoxGenerateWidget)
+											.OnSelectionChanged(this, &SMediaPlayerEditorViewer::HandleVideoTrackComboBoxSelectionChanged)
+											.OptionsSource(&VideoTracks)
+											.ToolTipText(LOCTEXT("CaptionTrackToolTip", "Select the video track to be played"))
+											[
+												SNew(STextBlock)
+													.Text(this, &SMediaPlayerEditorViewer::HandleVideoTrackComboBoxText)
+											]
+									]
+							]
 
 						+ SHorizontalBox::Slot()
 							.AutoWidth()
-							.VAlign(VAlign_Center)
+							.Padding(16.0f, 0.0f, 0.0f, 0.0f)
 							[
-								SNew(STextBlock)
-									.Text(LOCTEXT("CaptionTrackLabel", "Caption Track:"))
-							]
+								// audio track selector
+								SNew(SHorizontalBox)
 
-						+ SHorizontalBox::Slot()
-							.FillWidth(1.0f)
-							.Padding(4.0f, 0.0f, 0.0f, 0.0f)
-							.VAlign(VAlign_Center)
-							[
-								SAssignNew(CaptionTrackComboBox, SComboBox<IMediaTrackPtr>)
-									.OnGenerateWidget(this, &SMediaPlayerEditorViewer::HandleCaptionTrackComboBoxGenerateWidget)
-									.OnSelectionChanged(this, &SMediaPlayerEditorViewer::HandleCaptionTrackComboBoxSelectionChanged)
-									.OptionsSource(&CaptionTracks)
-									.ToolTipText(LOCTEXT("CaptionTrackToolTip", "Select the caption text to be displayed"))
+								+ SHorizontalBox::Slot()
+									.AutoWidth()
+									.VAlign(VAlign_Center)
 									[
 										SNew(STextBlock)
-											.Text(this, &SMediaPlayerEditorViewer::HandleCaptionTrackComboBoxText)
+											.Text(LOCTEXT("AudioTrackLabel", "Audio Track:"))
 									]
-							]
-					]
 
-				+ SHorizontalBox::Slot()
-					.FillWidth(1.0f)
-					.HAlign(HAlign_Right)
-					.Padding(16.0f, 0.0f, 0.0f, 0.0f)
-					[
-						// playback rate spin box
-						SNew(SHorizontalBox)
+								+ SHorizontalBox::Slot()
+									.FillWidth(1.0f)
+									.Padding(4.0f, 0.0f, 0.0f, 0.0f)
+									.VAlign(VAlign_Center)
+									[
+										SAssignNew(AudioTrackComboBox, SComboBox<IMediaTrackPtr>)
+											.OnGenerateWidget(this, &SMediaPlayerEditorViewer::HandleAudioTrackComboBoxGenerateWidget)
+											.OnSelectionChanged(this, &SMediaPlayerEditorViewer::HandleAudioTrackComboBoxSelectionChanged)
+											.OptionsSource(&AudioTracks)
+											.ToolTipText(LOCTEXT("CaptionTrackToolTip", "Select the audio track to be played"))
+											[
+												SNew(STextBlock)
+													.Text(this, &SMediaPlayerEditorViewer::HandleAudioTrackComboBoxText)
+											]
+									]
+							]
 
 						+ SHorizontalBox::Slot()
 							.AutoWidth()
-							.VAlign(VAlign_Center)
+							.Padding(16.0f, 0.0f, 0.0f, 0.0f)
 							[
-								SNew(STextBlock)
-									.Text(LOCTEXT("CaptionTrackLabel", "Playback Rate:"))
+								// caption track selector
+								SNew(SHorizontalBox)
+
+								+ SHorizontalBox::Slot()
+									.AutoWidth()
+									.VAlign(VAlign_Center)
+									[
+										SNew(STextBlock)
+											.Text(LOCTEXT("CaptionTrackLabel", "Caption Track:"))
+									]
+
+								+ SHorizontalBox::Slot()
+									.FillWidth(1.0f)
+									.Padding(4.0f, 0.0f, 0.0f, 0.0f)
+									.VAlign(VAlign_Center)
+									[
+										SAssignNew(CaptionTrackComboBox, SComboBox<IMediaTrackPtr>)
+											.OnGenerateWidget(this, &SMediaPlayerEditorViewer::HandleCaptionTrackComboBoxGenerateWidget)
+											.OnSelectionChanged(this, &SMediaPlayerEditorViewer::HandleCaptionTrackComboBoxSelectionChanged)
+											.OptionsSource(&CaptionTracks)
+											.ToolTipText(LOCTEXT("CaptionTrackToolTip", "Select the caption text to be displayed"))
+											[
+												SNew(STextBlock)
+													.Text(this, &SMediaPlayerEditorViewer::HandleCaptionTrackComboBoxText)
+											]
+									]
 							]
 
 						+ SHorizontalBox::Slot()
 							.FillWidth(1.0f)
-							.Padding(4.0f, 0.0f, 0.0f, 0.0f)
-							.VAlign(VAlign_Center)
+							.HAlign(HAlign_Right)
+							.Padding(16.0f, 0.0f, 0.0f, 0.0f)
 							[
-								SNew(SNumericEntryBox<float>)
-									.Delta(1.0f)
-									.MaxValue(this, &SMediaPlayerEditorViewer::HandlePlaybackRateBoxMaxValue)
-									.MinValue(this, &SMediaPlayerEditorViewer::HandlePlaybackRateBoxMinValue)
-									.Value(this, &SMediaPlayerEditorViewer::HandlePlaybackRateSpinBoxValue)
-									.OnValueChanged(this, &SMediaPlayerEditorViewer::HandlePlaybackRateBoxValueChanged)
+								// playback rate spin box
+								SNew(SHorizontalBox)
+
+								+ SHorizontalBox::Slot()
+									.AutoWidth()
+									.VAlign(VAlign_Center)
+									[
+										SNew(STextBlock)
+											.Text(LOCTEXT("CaptionTrackLabel", "Playback Rate:"))
+									]
+
+								+ SHorizontalBox::Slot()
+									.FillWidth(1.0f)
+									.Padding(4.0f, 0.0f, 0.0f, 0.0f)
+									.VAlign(VAlign_Center)
+									[
+										SNew(SNumericEntryBox<float>)
+											.Delta(1.0f)
+											.MaxValue(this, &SMediaPlayerEditorViewer::HandlePlaybackRateBoxMaxValue)
+											.MinValue(this, &SMediaPlayerEditorViewer::HandlePlaybackRateBoxMinValue)
+											.Value(this, &SMediaPlayerEditorViewer::HandlePlaybackRateSpinBoxValue)
+											.OnValueChanged(this, &SMediaPlayerEditorViewer::HandlePlaybackRateBoxValueChanged)
+									]
 							]
 					]
 			]
@@ -179,8 +208,9 @@ void SMediaPlayerEditorViewer::Construct( const FArguments& InArgs, UMediaPlayer
 				+ SOverlay::Slot()
 					[
 						// movie viewport
-						SAssignNew(ViewportWidget, SViewport)
+						SNew(SViewport)
 							.EnableGammaCorrection(false)
+							.ViewportInterface(Viewport)
 					]
 
 				+ SOverlay::Slot()
@@ -250,9 +280,6 @@ void SMediaPlayerEditorViewer::Construct( const FArguments& InArgs, UMediaPlayer
 			]
 	];
 
-	Viewport = MakeShareable(new FMediaPlayerEditorViewport());
-	ViewportWidget->SetViewportInterface(Viewport.ToSharedRef());
-
 	MediaPlayer->OnMediaChanged().AddRaw(this, &SMediaPlayerEditorViewer::HandleMediaPlayerMediaChanged);
 
 	ReloadMediaPlayer();
@@ -433,6 +460,12 @@ void SMediaPlayerEditorViewer::HandleMediaPlayerMediaChanged()
 }
 
 
+EVisibility SMediaPlayerEditorViewer::HandleNoMediaSelectedTextVisibility() const
+{
+	return MediaPlayer->GetUrl().IsEmpty() ? EVisibility::Visible : EVisibility::Hidden;
+}
+
+
 FText SMediaPlayerEditorViewer::HandleOverlayCaptionText() const
 {
 	TSharedPtr<TArray<uint8>, ESPMode::ThreadSafe> CurrentCaption = CaptionBuffer->GetCurrentSample();
@@ -448,6 +481,11 @@ FText SMediaPlayerEditorViewer::HandleOverlayCaptionText() const
 
 FText SMediaPlayerEditorViewer::HandleOverlayStateText() const
 {
+	if (MediaPlayer->GetUrl().IsEmpty())
+	{
+		return LOCTEXT("StateOverlayNoMedia", "No Media");
+	}
+
 	if (MediaPlayer->IsPaused())
 	{
 		return LOCTEXT("StateOverlayPaused", "Paused");
@@ -563,6 +601,12 @@ FText SMediaPlayerEditorViewer::HandleRemainingTimeTextBlockText() const
 }
 
 
+EVisibility SMediaPlayerEditorViewer::HandleTrackSelectionBoxVisibility() const
+{
+	return MediaPlayer->GetUrl().IsEmpty() ? EVisibility::Hidden : EVisibility::Visible;
+}
+
+
 TSharedRef<SWidget> SMediaPlayerEditorViewer::HandleVideoTrackComboBoxGenerateWidget( IMediaTrackPtr Value ) const
 {
 	return SNew(STextBlock)
diff --git a/Engine/Plugins/Media/MediaPlayerEditor/Source/MediaPlayerEditor/Private/Widgets/SMediaPlayerEditorViewer.h b/Engine/Plugins/Media/MediaPlayerEditor/Source/MediaPlayerEditor/Private/Widgets/SMediaPlayerEditorViewer.h
--- a/Engine/Plugins/Media/MediaPlayerEditor/Source/MediaPlayerEditor/Private/Widgets/SMediaPlayerEditorViewer.h
+++ b/Engine/Plugins/Media/MediaPlayerEditor/Source/MediaPlayerEditor/Private/Widgets/SMediaPlayerEditorViewer.h
@@ -64,6 +64,9 @@ private:
 	/** Callback for when the media player's media has changed. */
 	void HandleMediaPlayerMediaChanged();
 
+	/** Callback for getting the visibility of the 'No media selected' warning text. */
+	EVisibility HandleNoMediaSelectedTextVisibility() const;
+
 	/** Callback for getting the text of the movie captions. */
 	FText HandleOverlayCaptionText() const;
 
@@ -100,6 +103,9 @@ private:
 	/** Callback for getting the text of the 'Remaining Time' text block. */
 	FText HandleRemainingTimeTextBlockText() const;
 
+	/** Callback for getting the visibility of the track selection box. */
+	EVisibility HandleTrackSelectionBoxVisibility() const;
+
 	/** Callback for generating a widget for a value in the video track combo box. */
 	TSharedRef<SWidget> HandleVideoTrackComboBoxGenerateWidget( IMediaTrackPtr Value ) const;
 
