MediaPlayerEditor: Include what you use

diff --git a/Engine/Plugins/Media/MediaPlayerEditor/Source/MediaPlayerEditor/Private/Factories/MediaPlayerFactory.cpp b/Engine/Plugins/Media/MediaPlayerEditor/Source/MediaPlayerEditor/Private/Factories/MediaPlayerFactory.cpp
--- a/Engine/Plugins/Media/MediaPlayerEditor/Source/MediaPlayerEditor/Private/Factories/MediaPlayerFactory.cpp
+++ b/Engine/Plugins/Media/MediaPlayerEditor/Source/MediaPlayerEditor/Private/Factories/MediaPlayerFactory.cpp
@@ -1,6 +1,8 @@
 // Copyright 1998-2014 Epic Games, Inc. All Rights Reserved.
 
 #include "MediaPlayerEditorPrivatePCH.h"
+#include "IMediaModule.h"
+#include "IMediaPlayerFactory.h"
 
 
 /* UMediaPlayerFactory structors
diff --git a/Engine/Plugins/Media/MediaPlayerEditor/Source/MediaPlayerEditor/Private/MediaPlayerEditorModule.cpp b/Engine/Plugins/Media/MediaPlayerEditor/Source/MediaPlayerEditor/Private/MediaPlayerEditorModule.cpp
--- a/Engine/Plugins/Media/MediaPlayerEditor/Source/MediaPlayerEditor/Private/MediaPlayerEditorModule.cpp
+++ b/Engine/Plugins/Media/MediaPlayerEditor/Source/MediaPlayerEditor/Private/MediaPlayerEditorModule.cpp
@@ -2,7 +2,7 @@
 
 
 #include "MediaPlayerEditorPrivatePCH.h"
-#include "ModuleManager.h"
+#include "ModuleInterface.h"
 
 
 #define LOCTEXT_NAMESPACE "FMediaPlayerEditorModule"
diff --git a/Engine/Plugins/Media/MediaPlayerEditor/Source/MediaPlayerEditor/Private/MediaPlayerEditorPrivatePCH.h b/Engine/Plugins/Media/MediaPlayerEditor/Source/MediaPlayerEditor/Private/MediaPlayerEditorPrivatePCH.h
--- a/Engine/Plugins/Media/MediaPlayerEditor/Source/MediaPlayerEditor/Private/MediaPlayerEditorPrivatePCH.h
+++ b/Engine/Plugins/Media/MediaPlayerEditor/Source/MediaPlayerEditor/Private/MediaPlayerEditorPrivatePCH.h
@@ -13,11 +13,12 @@
 #include "SlateStyle.h"
 #include "EditorStyle.h"
 #include "IDetailsView.h"
-#include "Media.h"
+#include "IMediaTrack.h"
 #include "MediaPlayer.h"
+#include "MediaSampleBuffer.h"
 #include "MediaSoundWave.h"
 #include "MediaTexture.h"
-#include "ModuleInterface.h"
+#include "ModuleManager.h"
 #include "PropertyEditorModule.h"
 #include "SceneViewport.h"
 #include "SlateTextures.h"
diff --git a/Engine/Plugins/Media/MediaPlayerEditor/Source/MediaPlayerEditor/Private/Models/MediaPlayerEditorViewport.cpp b/Engine/Plugins/Media/MediaPlayerEditor/Source/MediaPlayerEditor/Private/Models/MediaPlayerEditorViewport.cpp
new file mode 100644
--- /dev/null
+++ b/Engine/Plugins/Media/MediaPlayerEditor/Source/MediaPlayerEditor/Private/Models/MediaPlayerEditorViewport.cpp
@@ -0,0 +1,96 @@
+// Copyright 1998-2014 Epic Games, Inc. All Rights Reserved.
+
+#include "MediaPlayerEditorPrivatePCH.h"
+#include "IMediaTrackVideoDetails.h"
+
+
+/* FMediaPlayerEditorViewport structors
+*****************************************************************************/
+
+FMediaPlayerEditorViewport::FMediaPlayerEditorViewport()
+: EditorTexture(nullptr)
+, SlateTexture(nullptr)
+{ }
+
+
+FMediaPlayerEditorViewport::~FMediaPlayerEditorViewport()
+{
+	ReleaseResources();
+}
+
+
+/* FMediaPlayerEditorViewport interface
+*****************************************************************************/
+
+void FMediaPlayerEditorViewport::Initialize(const IMediaTrackPtr& VideoTrack)
+{
+	ReleaseResources();
+
+	if (VideoTrack.IsValid())
+	{
+		check(VideoTrack->GetType() == EMediaTrackTypes::Video);
+
+		const bool bCreateEmptyTexture = true;
+		const FIntPoint VideoDimensions = VideoTrack->GetVideoDetails().GetDimensions();
+
+		SlateTexture = new FSlateTexture2DRHIRef(VideoDimensions.X, VideoDimensions.Y, PF_B8G8R8A8, nullptr, TexCreate_Dynamic, bCreateEmptyTexture);
+		EditorTexture = new FMediaPlayerEditorTexture(SlateTexture, VideoTrack.ToSharedRef());
+
+		ENQUEUE_UNIQUE_RENDER_COMMAND_ONEPARAMETER(
+			RegisterMediaPlayerEditorTexture,
+			FMediaPlayerEditorTexture*, EditorTextureParam, EditorTexture,
+			{
+				EditorTextureParam->Register();
+			}
+		);
+	}
+}
+
+
+/* ISlateViewport interface
+*****************************************************************************/
+
+FIntPoint FMediaPlayerEditorViewport::GetSize() const
+{
+	return (SlateTexture != nullptr)
+		? FIntPoint(SlateTexture->GetWidth(), SlateTexture->GetHeight())
+		: FIntPoint();
+}
+
+
+class FSlateShaderResource* FMediaPlayerEditorViewport::GetViewportRenderTargetTexture() const
+{
+	return SlateTexture;
+}
+
+
+bool FMediaPlayerEditorViewport::RequiresVsync() const
+{
+	return false;
+}
+
+
+/* FMediaPlayerEditorViewport implementation
+*****************************************************************************/
+
+void FMediaPlayerEditorViewport::ReleaseResources()
+{
+	if (SlateTexture != nullptr)
+	{
+		ENQUEUE_UNIQUE_RENDER_COMMAND_TWOPARAMETER(
+			ReleaseMediaPlayerEditorResources,
+			FMediaPlayerEditorTexture*, EditorTextureParam, EditorTexture,
+			FSlateTexture2DRHIRef*, SlateTextureParam, SlateTexture,
+			{
+			EditorTextureParam->Unregister();
+			delete EditorTextureParam;
+
+			SlateTextureParam->ReleaseResource();
+			delete SlateTextureParam;
+		}
+		);
+
+		EditorTexture = nullptr;
+		SlateTexture = nullptr;
+	}
+}
diff --git a/Engine/Plugins/Media/MediaPlayerEditor/Source/MediaPlayerEditor/Private/Models/MediaPlayerEditorViewport.h b/Engine/Plugins/Media/MediaPlayerEditor/Source/MediaPlayerEditor/Private/Models/MediaPlayerEditorViewport.h
--- a/Engine/Plugins/Media/MediaPlayerEditor/Source/MediaPlayerEditor/Private/Models/MediaPlayerEditorViewport.h
+++ b/Engine/Plugins/Media/MediaPlayerEditor/Source/MediaPlayerEditor/Private/Models/MediaPlayerEditorViewport.h
@@ -9,16 +9,10 @@ class FMediaPlayerEditorViewport
 public:
 
 	/** Default constructor. */
-	FMediaPlayerEditorViewport()
-		: EditorTexture(nullptr)
-		, SlateTexture(nullptr)
-	{ }
+	FMediaPlayerEditorViewport();
 
 	/** Destructor. */
-	~FMediaPlayerEditorViewport()
-	{
-		ReleaseResources();
-	}
+	~FMediaPlayerEditorViewport();
 
 public:
 
@@ -27,74 +21,20 @@ public:
 	 *
 	 * @param VideoTrack The video track to render to the viewport.
 	 */
-	void Initialize( const IMediaTrackPtr& VideoTrack )
-	{
-		ReleaseResources();
-
-		if (VideoTrack.IsValid())
-		{
-			check(VideoTrack->GetType() == EMediaTrackTypes::Video);
-
-			const bool bCreateEmptyTexture = true;
-			const FIntPoint VideoDimensions = VideoTrack->GetVideoDetails().GetDimensions();
-
-			SlateTexture = new FSlateTexture2DRHIRef(VideoDimensions.X, VideoDimensions.Y, PF_B8G8R8A8, nullptr, TexCreate_Dynamic, bCreateEmptyTexture);
-			EditorTexture = new FMediaPlayerEditorTexture(SlateTexture, VideoTrack.ToSharedRef());
-
-			ENQUEUE_UNIQUE_RENDER_COMMAND_ONEPARAMETER(
-				RegisterMediaPlayerEditorTexture,
-				FMediaPlayerEditorTexture*, EditorTextureParam, EditorTexture,
-				{
-					EditorTextureParam->Register();
-				}
-			);
-		}
-	}
-
-protected:
-
-	void ReleaseResources()
-	{
-		if (SlateTexture != nullptr)
-		{
-			ENQUEUE_UNIQUE_RENDER_COMMAND_TWOPARAMETER( 
-				ReleaseMediaPlayerEditorResources,
-				FMediaPlayerEditorTexture*, EditorTextureParam, EditorTexture,
-				FSlateTexture2DRHIRef*, SlateTextureParam, SlateTexture,
-				{
-					EditorTextureParam->Unregister();
-					delete EditorTextureParam;
-
-					SlateTextureParam->ReleaseResource();
-					delete SlateTextureParam;
-				}
-			);
-
-			EditorTexture = nullptr;
-			SlateTexture = nullptr;
-		}
-	}
+	void Initialize(const IMediaTrackPtr& VideoTrack);
 
 public:
 
 	// ISlateViewport interface
 
-	virtual FIntPoint GetSize() const override
-	{
-		return (SlateTexture != nullptr)
-			? FIntPoint(SlateTexture->GetWidth(), SlateTexture->GetHeight())
-			: FIntPoint();
-	}
+	virtual FIntPoint GetSize() const override;
+	virtual class FSlateShaderResource* GetViewportRenderTargetTexture() const override;
+	virtual bool RequiresVsync() const override;
 
-	virtual class FSlateShaderResource* GetViewportRenderTargetTexture() const override
-	{
-		return SlateTexture;
-	}
+protected:
 
-	virtual bool RequiresVsync() const override
-	{
-		return false;
-	}
+	/** Releases the resources associated with this viewport. */
+	void ReleaseResources();
 
 private:
 
diff --git a/Engine/Plugins/Media/MediaPlayerEditor/Source/MediaPlayerEditor/Private/Widgets/SMediaPlayerEditorViewer.cpp b/Engine/Plugins/Media/MediaPlayerEditor/Source/MediaPlayerEditor/Private/Widgets/SMediaPlayerEditorViewer.cpp
--- a/Engine/Plugins/Media/MediaPlayerEditor/Source/MediaPlayerEditor/Private/Widgets/SMediaPlayerEditorViewer.cpp
+++ b/Engine/Plugins/Media/MediaPlayerEditor/Source/MediaPlayerEditor/Private/Widgets/SMediaPlayerEditorViewer.cpp
@@ -1,6 +1,7 @@
 // Copyright 1998-2014 Epic Games, Inc. All Rights Reserved.
 
 #include "MediaPlayerEditorPrivatePCH.h"
+#include "IMediaPlayer.h"
 #include "SNumericEntryBox.h"
 
 
