Media: Fixed media player not rendering anything (Jira UE-9034)

diff --git a/Engine/Plugins/Media/MediaPlayerEditor/Source/MediaPlayerEditor/Private/Models/MediaPlayerEditorViewport.cpp b/Engine/Plugins/Media/MediaPlayerEditor/Source/MediaPlayerEditor/Private/Models/MediaPlayerEditorViewport.cpp
--- a/Engine/Plugins/Media/MediaPlayerEditor/Source/MediaPlayerEditor/Private/Models/MediaPlayerEditorViewport.cpp
+++ b/Engine/Plugins/Media/MediaPlayerEditor/Source/MediaPlayerEditor/Private/Models/MediaPlayerEditorViewport.cpp
@@ -33,7 +33,7 @@ void FMediaPlayerEditorViewport::Initialize(const IMediaTrackPtr& VideoTrack)
 		const bool bCreateEmptyTexture = true;
 		const FIntPoint VideoDimensions = VideoTrack->GetVideoDetails().GetDimensions();
 
-		SlateTexture = new FSlateTexture2DRHIRef(VideoDimensions.X, VideoDimensions.Y, PF_B8G8R8A8, nullptr, TexCreate_Dynamic, bCreateEmptyTexture);
+		SlateTexture = new FSlateTexture2DRHIRef(VideoDimensions.X, VideoDimensions.Y, PF_B8G8R8A8, nullptr, TexCreate_Dynamic | TexCreate_RenderTargetable, bCreateEmptyTexture);
 		EditorTexture = new FMediaPlayerEditorTexture(SlateTexture, VideoTrack.ToSharedRef());
 
 		ENQUEUE_UNIQUE_RENDER_COMMAND_ONEPARAMETER(
