Media: Disabling all controls in MediaPlayerEditor if media is not playable

diff --git a/Engine/Plugins/Media/MediaPlayerEditor/Source/MediaPlayerEditor/Private/MediaPlayerEditorToolkit.cpp b/Engine/Plugins/Media/MediaPlayerEditor/Source/MediaPlayerEditor/Private/MediaPlayerEditorToolkit.cpp
--- a/Engine/Plugins/Media/MediaPlayerEditor/Source/MediaPlayerEditor/Private/MediaPlayerEditorToolkit.cpp
+++ b/Engine/Plugins/Media/MediaPlayerEditor/Source/MediaPlayerEditor/Private/MediaPlayerEditorToolkit.cpp
@@ -281,7 +281,7 @@ float FMediaPlayerEditorToolkit::GetReverseRate() const
 
 bool FMediaPlayerEditorToolkit::HandleForwardMediaActionCanExecute() const
 {
-	return MediaPlayer->SupportsRate(GetForwardRate(), false);
+	return MediaPlayer->CanPlay() && MediaPlayer->SupportsRate(GetForwardRate(), false);
 }
 
 
@@ -317,7 +317,7 @@ void FMediaPlayerEditorToolkit::HandlePlayMediaActionExecute()
 
 bool FMediaPlayerEditorToolkit::HandleReverseMediaActionCanExecute() const
 {
-	return MediaPlayer->SupportsRate(GetReverseRate(), false);
+	return MediaPlayer->CanPlay() && MediaPlayer->SupportsRate(GetReverseRate(), false);
 }
 
 
@@ -329,7 +329,7 @@ void FMediaPlayerEditorToolkit::HandleReverseMediaActionExecute()
 
 bool FMediaPlayerEditorToolkit::HandleRewindMediaActionCanExecute() const
 {
-	return MediaPlayer->GetTime() > FTimespan::Zero();
+	return MediaPlayer->CanPlay() && MediaPlayer->GetTime() > FTimespan::Zero();
 }
 
 
