Media: Fixed regression: Video playback choppy due to Slate sleep mechanism

diff --git a/Engine/Plugins/Media/MediaPlayerEditor/Source/MediaPlayerEditor/Private/Widgets/SMediaPlayerEditorViewer.cpp b/Engine/Plugins/Media/MediaPlayerEditor/Source/MediaPlayerEditor/Private/Widgets/SMediaPlayerEditorViewer.cpp
--- a/Engine/Plugins/Media/MediaPlayerEditor/Source/MediaPlayerEditor/Private/Widgets/SMediaPlayerEditorViewer.cpp
+++ b/Engine/Plugins/Media/MediaPlayerEditor/Source/MediaPlayerEditor/Private/Widgets/SMediaPlayerEditorViewer.cpp
@@ -282,7 +282,7 @@ void SMediaPlayerEditorViewer::Construct( const FArguments& InArgs, UMediaPlayer
 	];
 
 	MediaPlayer->OnMediaChanged().AddRaw(this, &SMediaPlayerEditorViewer::HandleMediaPlayerMediaChanged);
-
+	RegisterActiveTimer(0.0f, FWidgetActiveTimerDelegate::CreateSP(this, &SMediaPlayerEditorViewer::HandleActiveTimer));
 	ReloadMediaPlayer();
 }
 
@@ -381,6 +381,12 @@ void SMediaPlayerEditorViewer::ReloadMediaPlayer()
 /* SMediaPlayerEditorPlayer callbacks
  *****************************************************************************/
 
+EActiveTimerReturnType SMediaPlayerEditorViewer::HandleActiveTimer(double InCurrentTime, float InDeltaTime)
+{
+	return EActiveTimerReturnType::Continue;
+}
+
+
 TSharedRef<SWidget> SMediaPlayerEditorViewer::HandleAudioTrackComboBoxGenerateWidget( IMediaTrackPtr Value ) const
 {
 	return SNew(STextBlock)
diff --git a/Engine/Plugins/Media/MediaPlayerEditor/Source/MediaPlayerEditor/Private/Widgets/SMediaPlayerEditorViewer.h b/Engine/Plugins/Media/MediaPlayerEditor/Source/MediaPlayerEditor/Private/Widgets/SMediaPlayerEditorViewer.h
--- a/Engine/Plugins/Media/MediaPlayerEditor/Source/MediaPlayerEditor/Private/Widgets/SMediaPlayerEditorViewer.h
+++ b/Engine/Plugins/Media/MediaPlayerEditor/Source/MediaPlayerEditor/Private/Widgets/SMediaPlayerEditorViewer.h
@@ -40,6 +40,9 @@ protected:
 
 private:
 
+	/** Callback for the active timer. */
+	EActiveTimerReturnType HandleActiveTimer(double InCurrentTime, float InDeltaTime);
+
 	/** Callback for generating a widget for a value in the audio track combo box. */
 	TSharedRef<SWidget> HandleAudioTrackComboBoxGenerateWidget( IMediaTrackPtr Value ) const;
 
