Media: Fixed playback rate label in MediaPlayerEditor

diff --git a/Engine/Plugins/Media/MediaPlayerEditor/Source/MediaPlayerEditor/Private/Widgets/SMediaPlayerEditorViewer.cpp b/Engine/Plugins/Media/MediaPlayerEditor/Source/MediaPlayerEditor/Private/Widgets/SMediaPlayerEditorViewer.cpp
--- a/Engine/Plugins/Media/MediaPlayerEditor/Source/MediaPlayerEditor/Private/Widgets/SMediaPlayerEditorViewer.cpp
+++ b/Engine/Plugins/Media/MediaPlayerEditor/Source/MediaPlayerEditor/Private/Widgets/SMediaPlayerEditorViewer.cpp
@@ -181,7 +181,7 @@ void SMediaPlayerEditorViewer::Construct( const FArguments& InArgs, UMediaPlayer
 									.VAlign(VAlign_Center)
 									[
 										SNew(STextBlock)
-											.Text(LOCTEXT("CaptionTrackLabel", "Playback Rate:"))
+											.Text(LOCTEXT("PlaybackRateLabel", "Playback Rate:"))
 									]
 
 								+ SHorizontalBox::Slot()
