Media: Not showing remaining time in MediaPlayerEditor if unknown.

diff --git a/Engine/Plugins/Media/MediaPlayerEditor/Source/MediaPlayerEditor/Private/Widgets/SMediaPlayerEditorViewer.cpp b/Engine/Plugins/Media/MediaPlayerEditor/Source/MediaPlayerEditor/Private/Widgets/SMediaPlayerEditorViewer.cpp
--- a/Engine/Plugins/Media/MediaPlayerEditor/Source/MediaPlayerEditor/Private/Widgets/SMediaPlayerEditorViewer.cpp
+++ b/Engine/Plugins/Media/MediaPlayerEditor/Source/MediaPlayerEditor/Private/Widgets/SMediaPlayerEditorViewer.cpp
@@ -604,7 +604,14 @@ void SMediaPlayerEditorViewer::HandlePositionSliderValueChanged( float NewValue
 
 FText SMediaPlayerEditorViewer::HandleRemainingTimeTextBlockText() const
 {
-	return FText::AsTimespan(MediaPlayer->GetDuration() - MediaPlayer->GetTime());
+	FTimespan Remaining = MediaPlayer->GetDuration() - MediaPlayer->GetTime();
+
+	if (Remaining <= FTimespan::Zero())
+	{
+		return FText::GetEmpty();
+	}
+
+	return FText::AsTimespan(Remaining);
 }
 
 
