Media: Added support for tracks changing during playback or asynchronously after opening media

diff --git a/Engine/Plugins/Media/AndroidMedia/Source/AndroidMedia/Private/AndroidMediaPlayer.cpp b/Engine/Plugins/Media/AndroidMedia/Source/AndroidMedia/Private/AndroidMediaPlayer.cpp
--- a/Engine/Plugins/Media/AndroidMedia/Source/AndroidMedia/Private/AndroidMediaPlayer.cpp
+++ b/Engine/Plugins/Media/AndroidMedia/Source/AndroidMedia/Private/AndroidMediaPlayer.cpp
@@ -371,7 +371,10 @@ void FAndroidMediaPlayer::Close()
 		}
 		MediaUrl = FString();
 		MediaState = EMediaState::Idle;
+		
 		Tracks.Reset();
+
+		TracksChangedEvent.Broadcast();
 		ClosedEvent.Broadcast();
 	}
 }
@@ -525,6 +528,8 @@ bool FAndroidMediaPlayer::Open(const FString& Url)
 		{
 			Tracks.Add(MakeShareable(new AudioTrack(*this, Tracks.Num())));
 		}
+
+		TracksChangedEvent.Broadcast();
 	}
 	if (MediaState == EMediaState::Prepared)
 	{
diff --git a/Engine/Plugins/Media/AndroidMedia/Source/AndroidMedia/Private/AndroidMediaPlayer.h b/Engine/Plugins/Media/AndroidMedia/Source/AndroidMedia/Private/AndroidMediaPlayer.h
--- a/Engine/Plugins/Media/AndroidMedia/Source/AndroidMedia/Private/AndroidMediaPlayer.h
+++ b/Engine/Plugins/Media/AndroidMedia/Source/AndroidMedia/Private/AndroidMediaPlayer.h
@@ -61,18 +61,25 @@ public:
 		return OpenedEvent;
 	}
 
+	DECLARE_DERIVED_EVENT(
+		FWmfMediaPlayer, IMediaPlayer::FOnTracksChanged, FOnTracksChanged);
+	virtual FOnTracksChanged& OnTracksChanged() override
+	{
+		return TracksChangedEvent;
+	}
+
 	// FTickableObjectRenderThread
 
 	virtual void Tick(float DeltaTime) override;
 	virtual TStatId GetStatId() const override;
 	virtual bool IsTickable() const override;
 
-private:
-
 	class MediaTrack;
 	class VideoTrack;
 	class AudioTrack;
 
+private:
+
 	// Local state to track where the Java side media player is at.
 	enum class EMediaState
 	{
@@ -89,6 +96,9 @@ private:
 	// Holds an event delegate that is invoked when media has been opened.
 	FOnMediaOpened OpenedEvent;
 
+	/** Holds an event delegate that is invoked when the media tracks have changed. */
+	FOnTracksChanged TracksChangedEvent;
+
 	// The Java side media interface.
 	TSharedPtr<FJavaAndroidMediaPlayer> JavaMediaPlayer;
 
diff --git a/Engine/Plugins/Media/AvfMedia/Source/AvfMedia/Private/Player/AvfMediaPlayer.cpp b/Engine/Plugins/Media/AvfMedia/Source/AvfMedia/Private/Player/AvfMediaPlayer.cpp
--- a/Engine/Plugins/Media/AvfMedia/Source/AvfMedia/Private/Player/AvfMediaPlayer.cpp
+++ b/Engine/Plugins/Media/AvfMedia/Source/AvfMedia/Private/Player/AvfMediaPlayer.cpp
@@ -131,6 +131,7 @@ void FAvfMediaPlayer::Close()
     }
     
 	Tracks.Reset();
+	TracksChangedEvent.Broadcast();
     
     Duration = CurrentTime = FTimespan::Zero();
     
@@ -265,7 +266,7 @@ bool FAvfMediaPlayer::Open( const FString& Url )
                         {
                             AVAssetTrack* VideoTrack = [VideoTracks objectAtIndex: 0];
                             Tracks.Add( MakeShareable( new FAvfMediaVideoTrack(VideoTrack) ) );
-     
+
 #if PLATFORM_MAC
                             GameThreadCall(^{
 #elif PLATFORM_IOS
@@ -273,8 +274,9 @@ bool FAvfMediaPlayer::Open( const FString& Url )
                             [FIOSAsyncTask CreateTaskWithBlock : ^ bool(void){
 #endif
                                 // Displatch on the gamethread that we have opened the video.
+								TracksChangedEvent.Broadcast();
                                 OpenedEvent.Broadcast( CachedUrl );
-                         
+
 #if PLATFORM_IOS
                                 return true;
                             }];
diff --git a/Engine/Plugins/Media/AvfMedia/Source/AvfMedia/Private/Player/AvfMediaPlayer.h b/Engine/Plugins/Media/AvfMedia/Source/AvfMedia/Private/Player/AvfMediaPlayer.h
--- a/Engine/Plugins/Media/AvfMedia/Source/AvfMedia/Private/Player/AvfMediaPlayer.h
+++ b/Engine/Plugins/Media/AvfMedia/Source/AvfMedia/Private/Player/AvfMediaPlayer.h
@@ -84,6 +84,12 @@ public:
 		return OpenedEvent;
 	}
 
+	DECLARE_DERIVED_EVENT(FWmfMediaPlayer, IMediaPlayer::FOnTracksChanged, FOnTracksChanged);
+	virtual FOnTracksChanged& OnTracksChanged() override
+	{
+		return TracksChangedEvent;
+	}
+
 public:
 
     // FTickerObjectBase interface
@@ -110,8 +116,6 @@ private:
     /** The available media tracks. */
     TArray<IMediaTrackRef> Tracks;
 
-private:
-
     /** The duration of the media. */
     FTimespan Duration;
 
@@ -124,6 +128,9 @@ private:
     /** The current playback rate. */
     float CurrentRate;
 
+	/** Cocoa helper object we can use to keep track of ns property changes in our media items */
+	FMediaHelper* MediaHelper;
+
 private:
 
 	/** Holds an event delegate that is invoked when media is being unloaded. */
@@ -131,8 +138,7 @@ private:
 
 	/** Holds an event delegate that is invoked when media has been loaded. */
 	FOnMediaOpened OpenedEvent;
-	
-private:
-	/** Cocoa helper object we can use to keep track of ns property changes in our media items */
-	FMediaHelper* MediaHelper;
+
+	/** Holds an event delegate that is invoked when the media tracks have changed. */
+	FOnTracksChanged TracksChangedEvent;
 };
diff --git a/Engine/Plugins/Media/MediaPlayerEditor/Source/MediaPlayerEditor/Private/Widgets/SMediaPlayerEditorViewer.cpp b/Engine/Plugins/Media/MediaPlayerEditor/Source/MediaPlayerEditor/Private/Widgets/SMediaPlayerEditorViewer.cpp
--- a/Engine/Plugins/Media/MediaPlayerEditor/Source/MediaPlayerEditor/Private/Widgets/SMediaPlayerEditorViewer.cpp
+++ b/Engine/Plugins/Media/MediaPlayerEditor/Source/MediaPlayerEditor/Private/Widgets/SMediaPlayerEditorViewer.cpp
@@ -282,15 +282,17 @@ void SMediaPlayerEditorViewer::Construct( const FArguments& InArgs, UMediaPlayer
 	];
 
 	MediaPlayer->OnMediaChanged().AddRaw(this, &SMediaPlayerEditorViewer::HandleMediaPlayerMediaChanged);
+	MediaPlayer->OnTracksChanged().AddRaw(this, &SMediaPlayerEditorViewer::HandleMediaPlayerTracksChanged);
+
 	RegisterActiveTimer(0.0f, FWidgetActiveTimerDelegate::CreateSP(this, &SMediaPlayerEditorViewer::HandleActiveTimer));
-	ReloadMediaPlayer();
+	ReloadTracks();
 }
 
 
 /* SMediaPlayerEditorPlayer implementation
  *****************************************************************************/
 
-void SMediaPlayerEditorViewer::ReloadMediaPlayer()
+void SMediaPlayerEditorViewer::ReloadTracks()
 {
 	// clear track collections
 	IMediaTrackPtr SelectedCaptionTrack = CaptionTrackComboBox->GetSelectedItem();
@@ -398,7 +400,7 @@ void SMediaPlayerEditorViewer::HandleAudioTrackComboBoxSelectionChanged( IMediaT
 {
 	if (SelectInfo != ESelectInfo::Direct)
 	{
-		ReloadMediaPlayer();
+		ReloadTracks();
 	}
 }
 
@@ -432,7 +434,7 @@ void SMediaPlayerEditorViewer::HandleCaptionTrackComboBoxSelectionChanged( IMedi
 {
 	if (SelectInfo != ESelectInfo::Direct)
 	{
-		ReloadMediaPlayer();
+		ReloadTracks();
 	}
 }
 
@@ -463,7 +465,13 @@ FText SMediaPlayerEditorViewer::HandleElapsedTimeTextBlockText() const
 
 void SMediaPlayerEditorViewer::HandleMediaPlayerMediaChanged()
 {
-	ReloadMediaPlayer();
+	ReloadTracks();
+}
+
+
+void SMediaPlayerEditorViewer::HandleMediaPlayerTracksChanged()
+{
+	ReloadTracks();
 }
 
 
@@ -632,7 +640,7 @@ void SMediaPlayerEditorViewer::HandleVideoTrackComboBoxSelectionChanged( IMediaT
 {
 	if (SelectInfo != ESelectInfo::Direct)
 	{
-		ReloadMediaPlayer();
+		ReloadTracks();
 	}
 }
 
diff --git a/Engine/Plugins/Media/MediaPlayerEditor/Source/MediaPlayerEditor/Private/Widgets/SMediaPlayerEditorViewer.h b/Engine/Plugins/Media/MediaPlayerEditor/Source/MediaPlayerEditor/Private/Widgets/SMediaPlayerEditorViewer.h
--- a/Engine/Plugins/Media/MediaPlayerEditor/Source/MediaPlayerEditor/Private/Widgets/SMediaPlayerEditorViewer.h
+++ b/Engine/Plugins/Media/MediaPlayerEditor/Source/MediaPlayerEditor/Private/Widgets/SMediaPlayerEditorViewer.h
@@ -36,7 +36,7 @@ public:
 protected:
 
 	/** Reloads the media texture and updates all options. */
-	void ReloadMediaPlayer();
+	void ReloadTracks();
 
 private:
 
@@ -67,6 +67,9 @@ private:
 	/** Callback for when the media player's media has changed. */
 	void HandleMediaPlayerMediaChanged();
 
+	/** Callback for when the media player's tracks changed. */
+	void HandleMediaPlayerTracksChanged();
+
 	/** Callback for getting the visibility of the 'No media selected' warning text. */
 	EVisibility HandleNoMediaSelectedTextVisibility() const;
 
diff --git a/Engine/Plugins/Media/WmfMedia/Source/WmfMedia/Private/Player/WmfMediaPlayer.cpp b/Engine/Plugins/Media/WmfMedia/Source/WmfMedia/Private/Player/WmfMediaPlayer.cpp
--- a/Engine/Plugins/Media/WmfMedia/Source/WmfMedia/Private/Player/WmfMediaPlayer.cpp
+++ b/Engine/Plugins/Media/WmfMedia/Source/WmfMedia/Private/Player/WmfMediaPlayer.cpp
@@ -91,6 +91,7 @@ void FWmfMediaPlayer::Close()
 	Duration = 0;
 	MediaUrl = FString();
 
+	TracksChangedEvent.Broadcast();
 	ClosedEvent.Broadcast();
 }
 
@@ -455,6 +456,7 @@ bool FWmfMediaPlayer::InitializeMediaSession(IUnknown* SourceObject, const FStri
 	}
 
 	UE_LOG(LogWmfMedia, Verbose, TEXT("Added a total of %i tracks"), Tracks.Num());
+	TracksChangedEvent.Broadcast();
 
 	UINT64 PresentationDuration = 0;
 	PresentationDescriptor->GetUINT64(MF_PD_DURATION, &PresentationDuration);
diff --git a/Engine/Plugins/Media/WmfMedia/Source/WmfMedia/Private/Player/WmfMediaPlayer.h b/Engine/Plugins/Media/WmfMedia/Source/WmfMedia/Private/Player/WmfMediaPlayer.h
--- a/Engine/Plugins/Media/WmfMedia/Source/WmfMedia/Private/Player/WmfMediaPlayer.h
+++ b/Engine/Plugins/Media/WmfMedia/Source/WmfMedia/Private/Player/WmfMediaPlayer.h
@@ -62,6 +62,12 @@ public:
 		return OpenedEvent;
 	}
 
+	DECLARE_DERIVED_EVENT(FWmfMediaPlayer, IMediaPlayer::FOnTracksChanged, FOnTracksChanged);
+	virtual FOnTracksChanged& OnTracksChanged() override
+	{
+		return TracksChangedEvent;
+	}
+
 protected:
 
 	/**
@@ -113,6 +119,9 @@ private:
 
 	/** Holds an event delegate that is invoked when media has been opened. */
 	FOnMediaOpened OpenedEvent;
+
+	/** Holds an event delegate that is invoked when the media tracks have changed. */
+	FOnTracksChanged TracksChangedEvent;
 };
 
 
diff --git a/Engine/Source/Runtime/Media/Public/IMediaPlayer.h b/Engine/Source/Runtime/Media/Public/IMediaPlayer.h
--- a/Engine/Source/Runtime/Media/Public/IMediaPlayer.h
+++ b/Engine/Source/Runtime/Media/Public/IMediaPlayer.h
@@ -180,6 +180,10 @@ public:
 	DECLARE_EVENT_OneParam(IMediaPlayer, FOnMediaOpened, FString /*OpenedUrl*/)
 	virtual FOnMediaOpened& OnOpened() = 0;
 
+	/** Gets an event delegate that is invoked when the media tracks have changed. */
+	DECLARE_EVENT(IMediaPlayer, FOnTracksChanged)
+	virtual FOnTracksChanged& OnTracksChanged() = 0;
+
 public:
 
 	/**
diff --git a/Engine/Source/Runtime/MediaAssets/Private/Assets/MediaPlayer.cpp b/Engine/Source/Runtime/MediaAssets/Private/Assets/MediaPlayer.cpp
--- a/Engine/Source/Runtime/MediaAssets/Private/Assets/MediaPlayer.cpp
+++ b/Engine/Source/Runtime/MediaAssets/Private/Assets/MediaPlayer.cpp
@@ -229,6 +229,7 @@ void UMediaPlayer::InitializePlayer()
 			Player->Close();
 			Player->OnClosed().RemoveAll(this);
 			Player->OnOpened().RemoveAll(this);
+			Player->OnTracksChanged().RemoveAll(this);
 			Player.Reset();
 		}
 
@@ -252,8 +253,9 @@ void UMediaPlayer::InitializePlayer()
 			return;
 		}
 
-		Player->OnClosed().AddUObject(this, &UMediaPlayer::HandleMediaPlayerMediaClosed);
-		Player->OnOpened().AddUObject(this, &UMediaPlayer::HandleMediaPlayerMediaOpened);
+		Player->OnClosed().AddUObject(this, &UMediaPlayer::HandlePlayerMediaClosed);
+		Player->OnOpened().AddUObject(this, &UMediaPlayer::HandlePlayerMediaOpened);
+		Player->OnTracksChanged().AddUObject(this, &UMediaPlayer::HandlePlayerTracksChanged);
 
 		// open the new media file
 		bool OpenedSuccessfully = false;
@@ -290,15 +292,21 @@ void UMediaPlayer::InitializePlayer()
 /* UMediaPlayer callbacks
  *****************************************************************************/
 
-void UMediaPlayer::HandleMediaPlayerMediaClosed()
+void UMediaPlayer::HandlePlayerMediaClosed()
 {
 	MediaChangedEvent.Broadcast();
 	OnMediaClosed.Broadcast();
 }
 
 
-void UMediaPlayer::HandleMediaPlayerMediaOpened( FString OpenedUrl )
+void UMediaPlayer::HandlePlayerMediaOpened( FString OpenedUrl )
 {
 	MediaChangedEvent.Broadcast();
 	OnMediaOpened.Broadcast(OpenedUrl);
 }
+
+
+void UMediaPlayer::HandlePlayerTracksChanged()
+{
+	TracksChangedEvent.Broadcast();
+}
diff --git a/Engine/Source/Runtime/MediaAssets/Private/Assets/MediaSoundWave.cpp b/Engine/Source/Runtime/MediaAssets/Private/Assets/MediaSoundWave.cpp
--- a/Engine/Source/Runtime/MediaAssets/Private/Assets/MediaSoundWave.cpp
+++ b/Engine/Source/Runtime/MediaAssets/Private/Assets/MediaSoundWave.cpp
@@ -12,6 +12,7 @@
 UMediaSoundWave::UMediaSoundWave( const FObjectInitializer& ObjectInitializer )
 	: Super(ObjectInitializer)
 	, AudioQueue(MakeShareable(new FMediaSampleQueue))
+	, AudioTrackIndex(INDEX_NONE)
 {
 	bLooping = false;
 	bProcedural = true;
@@ -25,6 +26,11 @@ UMediaSoundWave::~UMediaSoundWave()
 	{
 		AudioTrack->RemoveSink(AudioQueue);
 	}
+
+	if (CurrentMediaPlayer != nullptr)
+	{
+		CurrentMediaPlayer->OnTracksChanged().RemoveAll(this);
+	}
 }
 
 
@@ -147,14 +153,14 @@ void UMediaSoundWave::InitializeTrack()
 	{
 		if (CurrentMediaPlayer != nullptr)
 		{
-			CurrentMediaPlayer->OnMediaChanged().RemoveAll(this);
+			CurrentMediaPlayer->OnTracksChanged().RemoveAll(this);
 		}
 
 		CurrentMediaPlayer = MediaPlayer;
 
 		if (MediaPlayer != nullptr)
 		{
-			MediaPlayer->OnMediaChanged().AddUObject(this, &UMediaSoundWave::HandleMediaPlayerMediaChanged);
+			MediaPlayer->OnTracksChanged().AddUObject(this, &UMediaSoundWave::HandleMediaPlayerTracksChanged);
 		}	
 	}
 
@@ -173,9 +179,7 @@ void UMediaSoundWave::InitializeTrack()
 
 		if (Player .IsValid())
 		{
-			AudioTrack = Player ->GetTrack(AudioTrackIndex, EMediaTrackTypes::Audio);
-
-			if (!AudioTrack.IsValid())
+			if (AudioTrackIndex == INDEX_NONE)
 			{
 				AudioTrack = Player->GetFirstTrack(EMediaTrackTypes::Audio);
 
@@ -184,6 +188,10 @@ void UMediaSoundWave::InitializeTrack()
 					AudioTrackIndex = AudioTrack->GetIndex();
 				}
 			}
+			else
+			{
+				AudioTrack = Player ->GetTrack(AudioTrackIndex, EMediaTrackTypes::Audio);
+			}
 		}
 	}
 
@@ -209,7 +217,7 @@ void UMediaSoundWave::InitializeTrack()
 /* UMediaSoundWave callbacks
  *****************************************************************************/
 
-void UMediaSoundWave::HandleMediaPlayerMediaChanged()
+void UMediaSoundWave::HandleMediaPlayerTracksChanged()
 {
 	InitializeTrack();
 }
diff --git a/Engine/Source/Runtime/MediaAssets/Private/Assets/MediaTexture.cpp b/Engine/Source/Runtime/MediaAssets/Private/Assets/MediaTexture.cpp
--- a/Engine/Source/Runtime/MediaAssets/Private/Assets/MediaTexture.cpp
+++ b/Engine/Source/Runtime/MediaAssets/Private/Assets/MediaTexture.cpp
@@ -14,6 +14,7 @@ UMediaTexture::UMediaTexture( const FObjectInitializer& ObjectInitializer )
 	, MediaPlayer(nullptr)
 	, CurrentMediaPlayer(nullptr)
 	, VideoBuffer(MakeShareable(new FMediaSampleBuffer))
+	, VideoTrackIndex(INDEX_NONE)
 {
 	NeverStream = true;
 
@@ -27,6 +28,11 @@ UMediaTexture::~UMediaTexture()
 	{
 		VideoTrack->RemoveSink(VideoBuffer);
 	}
+
+	if (CurrentMediaPlayer != nullptr)
+	{
+		CurrentMediaPlayer->OnTracksChanged().RemoveAll(this);
+	}
 }
 
 
@@ -171,14 +177,14 @@ void UMediaTexture::InitializeTrack()
 	{
 		if (CurrentMediaPlayer != nullptr)
 		{
-			CurrentMediaPlayer->OnMediaChanged().RemoveAll(this);
+			CurrentMediaPlayer->OnTracksChanged().RemoveAll(this);
 		}
 
 		CurrentMediaPlayer = MediaPlayer;
 
 		if (MediaPlayer != nullptr)
 		{
-			MediaPlayer->OnMediaChanged().AddUObject(this, &UMediaTexture::HandleMediaPlayerMediaChanged);
+			MediaPlayer->OnTracksChanged().AddUObject(this, &UMediaTexture::HandleMediaPlayerTracksChanged);
 		}	
 	}
 
@@ -196,9 +202,7 @@ void UMediaTexture::InitializeTrack()
 
 		if (Player.IsValid())
 		{
-			VideoTrack = Player->GetTrack(VideoTrackIndex, EMediaTrackTypes::Video);
-
-			if (!VideoTrack.IsValid())
+			if (VideoTrackIndex == INDEX_NONE)
 			{
 				VideoTrack = Player->GetFirstTrack(EMediaTrackTypes::Video);
 
@@ -207,6 +211,10 @@ void UMediaTexture::InitializeTrack()
 					VideoTrackIndex = VideoTrack->GetIndex();
 				}
 			}
+			else
+			{
+				VideoTrack = Player->GetTrack(VideoTrackIndex, EMediaTrackTypes::Video);
+			}
 		}
 	}
 
@@ -232,7 +240,7 @@ void UMediaTexture::InitializeTrack()
 /* UMediaTexture callbacks
  *****************************************************************************/
 
-void UMediaTexture::HandleMediaPlayerMediaChanged()
+void UMediaTexture::HandleMediaPlayerTracksChanged()
 {
 	InitializeTrack();
 }
diff --git a/Engine/Source/Runtime/MediaAssets/Public/MediaPlayer.h b/Engine/Source/Runtime/MediaAssets/Public/MediaPlayer.h
--- a/Engine/Source/Runtime/MediaAssets/Public/MediaPlayer.h
+++ b/Engine/Source/Runtime/MediaAssets/Public/MediaPlayer.h
@@ -239,6 +239,13 @@ public:
 		return MediaChangedEvent;
 	}
 
+	/** Gets an event delegate that is invoked when media tracks have changed. */
+	DECLARE_EVENT(UMediaPlayer, FOnTracksChanged)
+	FOnTracksChanged& OnTracksChanged()
+	{
+		return TracksChangedEvent;
+	}
+
 	/** Holds a delegate that is invoked when a media source has been closed. */
 	UPROPERTY(BlueprintAssignable, Category="Media|MediaPlayer")
 	FOnMediaPlayerMediaClosed OnMediaClosed;
@@ -288,11 +295,14 @@ protected:
 
 private:
 
-	/** Callback for when the media player has closed a media source. */
-	void HandleMediaPlayerMediaClosed();
+	/** Callback for when the player has closed a media source. */
+	void HandlePlayerMediaClosed();
 
-	/** Callback for when the media player has opened a new media source. */
-	void HandleMediaPlayerMediaOpened( FString OpenedUrl );
+	/** Callback for when the player has opened a new media source. */
+	void HandlePlayerMediaOpened( FString OpenedUrl );
+
+	/** Callback for when the player's tracks changed. */
+	void HandlePlayerTracksChanged();
 
 private:
 
@@ -306,4 +316,7 @@ private:
 
 	/** Holds a delegate that is executed when media has been opened or closed. */
 	FOnMediaChanged MediaChangedEvent;
+
+	/** Holds a delegate that is executed when media tracks have changed. */
+	FOnTracksChanged TracksChangedEvent;
 };
diff --git a/Engine/Source/Runtime/MediaAssets/Public/MediaSoundWave.h b/Engine/Source/Runtime/MediaAssets/Public/MediaSoundWave.h
--- a/Engine/Source/Runtime/MediaAssets/Public/MediaSoundWave.h
+++ b/Engine/Source/Runtime/MediaAssets/Public/MediaSoundWave.h
@@ -80,8 +80,8 @@ protected:
 
 private:
 
-	/** Callback for when the UMediaPlayer asset changed its media. */
-	void HandleMediaPlayerMediaChanged();
+	/** Callback for when the UMediaPlayer changed tracks. */
+	void HandleMediaPlayerTracksChanged();
 
 private:
 
diff --git a/Engine/Source/Runtime/MediaAssets/Public/MediaTexture.h b/Engine/Source/Runtime/MediaAssets/Public/MediaTexture.h
--- a/Engine/Source/Runtime/MediaAssets/Public/MediaTexture.h
+++ b/Engine/Source/Runtime/MediaAssets/Public/MediaTexture.h
@@ -129,8 +129,8 @@ protected:
 
 private:
 
-	/** Callback for when the UMediaPlayer asset changed its media. */
-	void HandleMediaPlayerMediaChanged();
+	/** Callback for when the UMediaPlayer changed tracks. */
+	void HandleMediaPlayerTracksChanged();
 
 private:
 
