Media: API refactor in preparation for performance improvements

This is an API breaking change. IMediaTrack has been renamed to IMediaStream. There are three new interfaces for tracks: IMediaAudioTrack, IMediaCaptionTrack and IMediaVideoTrack.

diff --git a/Engine/Plugins/Media/AndroidMedia/Source/AndroidMedia/Private/AndroidMediaPCH.h b/Engine/Plugins/Media/AndroidMedia/Source/AndroidMedia/Private/AndroidMediaPCH.h
--- a/Engine/Plugins/Media/AndroidMedia/Source/AndroidMedia/Private/AndroidMediaPCH.h
+++ b/Engine/Plugins/Media/AndroidMedia/Source/AndroidMedia/Private/AndroidMediaPCH.h
@@ -6,8 +6,8 @@
 #include "IMediaInfo.h"
 #include "IMediaPlayer.h"
 #include "IMediaSink.h"
-#include "IMediaTrack.h"
-#include "IMediaTrackAudioDetails.h"
-#include "IMediaTrackCaptionDetails.h"
-#include "IMediaTrackVideoDetails.h"
+#include "IMediaStream.h"
+#include "IMediaAudioTrack.h"
+#include "IMediaCaptionTrack.h"
+#include "IMediaVideoTrack.h"
 #include "ModuleManager.h"
diff --git a/Engine/Plugins/Media/AndroidMedia/Source/AndroidMedia/Private/AndroidMediaPlayer.cpp b/Engine/Plugins/Media/AndroidMedia/Source/AndroidMedia/Private/AndroidMediaPlayer.cpp
--- a/Engine/Plugins/Media/AndroidMedia/Source/AndroidMedia/Private/AndroidMediaPlayer.cpp
+++ b/Engine/Plugins/Media/AndroidMedia/Source/AndroidMedia/Private/AndroidMediaPlayer.cpp
@@ -14,7 +14,7 @@ DEFINE_LOG_CATEGORY_STATIC(LogAndroidMediaPlayer, Log, All);
 #define LOCTEXT_NAMESPACE "FAndroidMediaModule"
 
 class FAndroidMediaPlayer::MediaTrack
-	: public IMediaTrack
+	: public IMediaStream
 {
 public:
 
@@ -42,28 +42,6 @@ public:
 		return MediaPlayer.MediaState != EMediaState::Error;
 	}
 
-	virtual const IMediaTrackAudioDetails& GetAudioDetails() const
-	{
-		verify(false); // not an audio track
-		return *static_cast<IMediaTrackAudioDetails*>(nullptr);
-	}
-
-	virtual const IMediaTrackCaptionDetails& GetCaptionDetails() const
-	{
-		verify(false); // not a caption track
-		return *static_cast<IMediaTrackCaptionDetails*>(nullptr);
-	}
-
-	virtual FText GetDisplayName() const
-	{
-		return FText::Format(LOCTEXT("UnnamedTrackFormat", "Unnamed Track {0}"), FText::AsNumber(GetIndex()));
-	}
-
-	virtual uint32 GetIndex() const
-	{
-		return TrackIndex;
-	}
-
 	virtual FString GetLanguage() const
 	{
 		return TEXT("und");
@@ -74,18 +52,12 @@ public:
 		return FPaths::GetBaseFilename(MediaPlayer.MediaUrl);
 	}
 
-	virtual const IMediaTrackVideoDetails& GetVideoDetails() const
-	{
-		verify(false); // not a video track
-		return *static_cast<IMediaTrackVideoDetails*>(nullptr);
-	}
-
 	virtual bool IsEnabled() const
 	{
 		return TrackIsEnabled;
 	}
 
-	virtual bool IsMutuallyExclusive(const IMediaTrackRef& Other) const
+	virtual bool IsMutuallyExclusive(const IMediaStreamRef& Other) const
 	{
 		return false;
 	}
@@ -104,11 +76,6 @@ public:
 	{
 	}
 
-	virtual bool TickInRenderThread()
-	{
-		return false;
-	}
-
 	virtual void ProcessMediaSample(void* SampleBuffer, uint32 SampleSize, FTimespan SampleDuration, FTimespan SampleTime)
 	{
 		/**
@@ -150,7 +117,7 @@ protected:
 
 class FAndroidMediaPlayer::VideoTrack
 	: public FAndroidMediaPlayer::MediaTrack
-	, public IMediaTrackVideoDetails
+	, public IMediaVideoTrack
 {
 public:
 
@@ -180,14 +147,9 @@ public:
 		return false;
 	}
 
-	virtual EMediaTrackTypes GetType() const
+	virtual FText GetDisplayName() const
 	{
-		return EMediaTrackTypes::Video;
-	}
-
-	virtual const IMediaTrackVideoDetails& GetVideoDetails() const
-	{
-		return *this;
+		return FText::Format(LOCTEXT("VideoTrackNameFormat", "Video Track {0}"), FText::AsNumber(TrackIndex));
 	}
 
 	virtual uint32 GetBitRate() const
@@ -214,6 +176,11 @@ public:
 		return 30.0f;
 	}
 
+	virtual IMediaStream& GetStream() override
+	{
+		return *this;
+	}
+
 	virtual void Tick(float DeltaTime)
 	{
 		if (MediaPlayer.MediaState != EMediaState::Error)
@@ -237,11 +204,6 @@ public:
 		}
 	}
 
-	virtual bool TickInRenderThread()
-	{
-		return true;
-	}
-
 private:
 
 	int32 LastFramePosition;
@@ -249,7 +211,7 @@ private:
 
 class FAndroidMediaPlayer::AudioTrack
 	: public FAndroidMediaPlayer::MediaTrack
-	, public IMediaTrackAudioDetails
+	, public IMediaAudioTrack
 {
 public:
 
@@ -278,14 +240,9 @@ public:
 		return false;
 	}
 
-	virtual const IMediaTrackAudioDetails& GetAudioDetails() const
-	{
-		return *this;
-	}
-
-	virtual EMediaTrackTypes GetType() const
+	virtual FText GetDisplayName() const
 	{
-		return EMediaTrackTypes::Audio;
+		return FText::Format(LOCTEXT("AudioTrackNameFormat", "Audio Track {0}"), FText::AsNumber(TrackIndex));
 	}
 
 	virtual uint32 GetNumChannels() const
@@ -297,6 +254,11 @@ public:
 	{
 		return 0;
 	}
+
+	virtual IMediaStream& GetStream() override
+	{
+		return *this;
+	}
 };
 
 FAndroidMediaPlayer::FAndroidMediaPlayer()
@@ -372,13 +334,26 @@ void FAndroidMediaPlayer::Close()
 		MediaUrl = FString();
 		MediaState = EMediaState::Idle;
 		
-		Tracks.Reset();
+		AudioTracks.Reset();
+		CaptionTracks.Reset();
+		VideoTracks.Reset();
 
 		TracksChangedEvent.Broadcast();
 		ClosedEvent.Broadcast();
 	}
 }
 
+const TArray<IMediaAudioTrackRef>& FAndroidMediaPlayer::GetAudioTracks() const
+{
+	return AudioTracks;
+}
+
+
+const TArray<IMediaCaptionTrackRef>& FAndroidMediaPlayer::GetCaptionTracks() const
+{
+	return CaptionTracks;
+}
+
 const IMediaInfo& FAndroidMediaPlayer::GetMediaInfo() const
 {
 	return *this;
@@ -401,9 +376,9 @@ FTimespan FAndroidMediaPlayer::GetTime() const
 	}
 }
 
-const TArray<IMediaTrackRef>& FAndroidMediaPlayer::GetTracks() const
+const TArray<IMediaVideoTrackRef>& FAndroidMediaPlayer::GetVideoTracks() const
 {
-	return Tracks;
+	return VideoTracks;
 }
 
 bool FAndroidMediaPlayer::IsLooping() const
@@ -521,12 +496,12 @@ bool FAndroidMediaPlayer::Open(const FString& Url)
 		{
 			// For video we add video track and disable audio
 			JavaMediaPlayer->SetAudioEnabled(false);
-//			Tracks.Add(MakeShareable(new AudioTrack(*this, Tracks.Num())));
-			Tracks.Add(MakeShareable(new VideoTrack(*this, Tracks.Num())));
+//			AudioTracks.Add(MakeShareable(new AudioTrack(*this, AudioTracks.Num())));
+			VideoTracks.Add(MakeShareable(new VideoTrack(*this, VideoTracks.Num())));
 		}
 		else if (Extension.Equals(TEXT("aac"), ESearchCase::IgnoreCase))
 		{
-			Tracks.Add(MakeShareable(new AudioTrack(*this, Tracks.Num())));
+			AudioTracks.Add(MakeShareable(new AudioTrack(*this, AudioTracks.Num())));
 		}
 
 		TracksChangedEvent.Broadcast();
@@ -611,29 +586,31 @@ bool FAndroidMediaPlayer::SetRate(float Rate)
 	return false;
 }
 
-typedef TWeakPtr<IMediaTrack, ESPMode::ThreadSafe> IMediaTrackWeakPtr;
+typedef TWeakPtr<IMediaVideoTrack, ESPMode::ThreadSafe> IMediaVideoTrackWeakPtr;
 
 void FAndroidMediaPlayer::Tick(float DeltaTime)
 {
-	for (IMediaTrackRef track : Tracks)
+	for (IMediaAudioTrackRef AudioTrack : AudioTracks)
 	{
-		MediaTrack & TrackToTick = static_cast<MediaTrack&>(*track);
-		if (TrackToTick.TickInRenderThread())
-		{
-			ENQUEUE_UNIQUE_RENDER_COMMAND_TWOPARAMETER(
-				RenderTickMediaTrack, IMediaTrackWeakPtr, Track, track, float, DeltaT, DeltaTime,
+		static_cast<MediaTrack&>(AudioTrack->GetStream()).Tick(DeltaTime);
+	}
+
+	for (IMediaCaptionTrackRef CaptionTrack : CaptionTracks)
+	{
+		static_cast<MediaTrack&>(CaptionTrack->GetStream()).Tick(DeltaTime);
+	}
+
+	for (IMediaVideoTrackRef VideoTrack : VideoTracks)
+	{
+		ENQUEUE_UNIQUE_RENDER_COMMAND_TWOPARAMETER(
+			RenderTickMediaTrack, IMediaVideoTrackWeakPtr, Track, VideoTrack, float, DeltaT, DeltaTime,
+			{
+				IMediaVideoTrackPtr t = Track.Pin();
+				if (t.IsValid())
 				{
-					IMediaTrackPtr t = Track.Pin();
-					if (t.IsValid())
-					{
-						StaticCastSharedPtr<MediaTrack>(t)->Tick(DeltaT);
-					}
-				});
-		}
-		else
-		{
-			TrackToTick.Tick(DeltaTime);
-		}
+					static_cast<MediaTrack&>(t->GetStream()).Tick(DeltaT);
+				}
+			});
 	}
 }
 
diff --git a/Engine/Plugins/Media/AndroidMedia/Source/AndroidMedia/Private/AndroidMediaPlayer.h b/Engine/Plugins/Media/AndroidMedia/Source/AndroidMedia/Private/AndroidMediaPlayer.h
--- a/Engine/Plugins/Media/AndroidMedia/Source/AndroidMedia/Private/AndroidMediaPlayer.h
+++ b/Engine/Plugins/Media/AndroidMedia/Source/AndroidMedia/Private/AndroidMediaPlayer.h
@@ -15,9 +15,12 @@ class FAndroidMediaPlayer
 	, public FTickableGameObject
 {
 public:
+	
 	FAndroidMediaPlayer();
 	~FAndroidMediaPlayer();
 
+public:
+
 	// IMediaInfo interface
 
 	virtual FTimespan GetDuration() const override;
@@ -29,45 +32,47 @@ public:
 	virtual bool SupportsScrubbing() const override;
 	virtual bool SupportsSeeking() const override;
 
+public:
+
 	// IMediaPlayer interface
 
 	virtual void Close() override;
+	virtual const TArray<IMediaAudioTrackRef>& GetAudioTracks() const override;
+	virtual const TArray<IMediaCaptionTrackRef>& GetCaptionTracks() const override;
 	virtual const IMediaInfo& GetMediaInfo() const override;
 	virtual float GetRate() const override;
 	virtual FTimespan GetTime() const override;
-	virtual const TArray<IMediaTrackRef>& GetTracks() const override;
+	virtual const TArray<IMediaVideoTrackRef>& GetVideoTracks() const override;
 	virtual bool IsLooping() const override;
 	virtual bool IsPaused() const override;
 	virtual bool IsPlaying() const override;
 	virtual bool IsReady() const override;
 	virtual bool Open(const FString& Url) override;
-	virtual bool Open(const TSharedRef<FArchive, ESPMode::ThreadSafe>& Archive,
-		const FString& OriginalUrl) override;
+	virtual bool Open(const TSharedRef<FArchive, ESPMode::ThreadSafe>& Archive, const FString& OriginalUrl) override;
 	virtual bool Seek(const FTimespan& Time) override;
 	virtual bool SetLooping(bool Looping) override;
 	virtual bool SetRate(float Rate) override;
 
-	DECLARE_DERIVED_EVENT(
-		FAndroidMediaPlayer, IMediaPlayer::FOnMediaClosed, FOnMediaClosed);
+	DECLARE_DERIVED_EVENT(FAndroidMediaPlayer, IMediaPlayer::FOnMediaClosed, FOnMediaClosed);
 	virtual FOnMediaClosed& OnClosed() override
 	{
 		return ClosedEvent;
 	}
 
-	DECLARE_DERIVED_EVENT(
-		FAndroidMediaPlayer, IMediaPlayer::FOnMediaOpened, FOnMediaOpened);
+	DECLARE_DERIVED_EVENT(FAndroidMediaPlayer, IMediaPlayer::FOnMediaOpened, FOnMediaOpened);
 	virtual FOnMediaOpened& OnOpened() override
 	{
 		return OpenedEvent;
 	}
 
-	DECLARE_DERIVED_EVENT(
-		FWmfMediaPlayer, IMediaPlayer::FOnTracksChanged, FOnTracksChanged);
+	DECLARE_DERIVED_EVENT(FWmfMediaPlayer, IMediaPlayer::FOnTracksChanged, FOnTracksChanged);
 	virtual FOnTracksChanged& OnTracksChanged() override
 	{
 		return TracksChangedEvent;
 	}
 
+public:
+
 	// FTickableObjectRenderThread
 
 	virtual void Tick(float DeltaTime) override;
@@ -87,24 +92,30 @@ private:
 		Paused, Stopped, PlaybackCompleted, End, Error
 	};
 
-	// Our understanding of the state of the Java media player.
-	EMediaState MediaState;
+	/** The available audio tracks. */
+	TArray<IMediaAudioTrackRef> AudioTracks;
 
-	// Holds an event delegate that is invoked when media has been closed.
+	/** The available caption tracks. */
+	TArray<IMediaCaptionTrackRef> CaptionTracks;
+
+	/** Holds an event delegate that is invoked when media has been closed. */
 	FOnMediaClosed ClosedEvent;
 
-	// Holds an event delegate that is invoked when media has been opened.
+	/** The Java side media interface. */
+	TSharedPtr<FJavaAndroidMediaPlayer> JavaMediaPlayer;
+
+	/** Our understanding of the state of the Java media player. */
+	EMediaState MediaState;
+
+	/** Holds an event delegate that is invoked when media has been opened. */
 	FOnMediaOpened OpenedEvent;
 
 	/** Holds an event delegate that is invoked when the media tracks have changed. */
 	FOnTracksChanged TracksChangedEvent;
 
-	// The Java side media interface.
-	TSharedPtr<FJavaAndroidMediaPlayer> JavaMediaPlayer;
-
-	// Currently opened media.
+	/** Currently opened media. */
 	FString MediaUrl;
 
-	// The pseudo-tracks in the media.
-	TArray<IMediaTrackRef> Tracks;
+	/** The available video tracks. */
+	TArray<IMediaVideoTrackRef> VideoTracks;
 };
diff --git a/Engine/Plugins/Media/AvfMedia/Source/AvfMedia/Private/AvfMediaPrivatePCH.h b/Engine/Plugins/Media/AvfMedia/Source/AvfMedia/Private/AvfMediaPrivatePCH.h
--- a/Engine/Plugins/Media/AvfMedia/Source/AvfMedia/Private/AvfMediaPrivatePCH.h
+++ b/Engine/Plugins/Media/AvfMedia/Source/AvfMedia/Private/AvfMediaPrivatePCH.h
@@ -8,10 +8,10 @@
 
 #include "Core.h"
 #include "IMediaSink.h"
-#include "IMediaTrack.h"
-#include "IMediaTrackAudioDetails.h"
-#include "IMediaTrackCaptionDetails.h"
-#include "IMediaTrackVideoDetails.h"
+#include "IMediaTrackInfo.h"
+#include "IMediaAudioTrack.h"
+#include "IMediaCaptionTrack.h"
+#include "IMediaVideoTrack.h"
 
 #import <AVFoundation/AVFoundation.h>
 
diff --git a/Engine/Plugins/Media/AvfMedia/Source/AvfMedia/Private/Player/AvfMediaPlayer.cpp b/Engine/Plugins/Media/AvfMedia/Source/AvfMedia/Private/Player/AvfMediaPlayer.cpp
--- a/Engine/Plugins/Media/AvfMedia/Source/AvfMedia/Private/Player/AvfMediaPlayer.cpp
+++ b/Engine/Plugins/Media/AvfMedia/Source/AvfMedia/Private/Player/AvfMediaPlayer.cpp
@@ -130,7 +130,10 @@ void FAvfMediaPlayer::Close()
         MediaPlayer = nil;
     }
     
-	Tracks.Reset();
+	AudioTracks.Reset();
+	CaptionTracks.Reset();
+	VideoTracks.Reset();
+
 	TracksChangedEvent.Broadcast();
     
     Duration = CurrentTime = FTimespan::Zero();
@@ -154,6 +157,18 @@ void FAvfMediaPlayer::Close()
 }
 
 
+const TArray<IMediaAudioTrackRef>& FAvfMediaPlayer::GetAudioTracks() const
+{
+	return AudioTracks;
+}
+
+
+const TArray<IMediaCaptionTrackRef>& FAvfMediaPlayer::GetCaptionTracks() const
+{
+	return CaptionTracks;
+}
+
+
 const IMediaInfo& FAvfMediaPlayer::GetMediaInfo() const
 {
 	return *this;
@@ -172,9 +187,9 @@ FTimespan FAvfMediaPlayer::GetTime() const
 }
 
 
-const TArray<IMediaTrackRef>& FAvfMediaPlayer::GetTracks() const
+const TArray<IMediaVideoTrackRef>& FAvfMediaPlayer::GetVideoTracks() const
 {
-	return Tracks;
+	return VideoTracks;
 }
 
 
@@ -192,29 +207,58 @@ bool FAvfMediaPlayer::IsPaused() const
 
 bool FAvfMediaPlayer::IsPlaying() const
 {
-	return (MediaPlayer != nil) && !FMath::IsNearlyZero([MediaPlayer rate]) && (Tracks.Num() > 0) && (CurrentTime <= Duration);
+	if ((AudioTracks.Num() > 0) && (CaptionTracks.Num() > 0) && (VideoTracks.Num() > 0))
+	{
+		return false;
+	}
+
+	return (MediaPlayer != nil) && !FMath::IsNearlyZero([MediaPlayer rate]) && (CurrentTime <= Duration);
 }
 
 
 bool FAvfMediaPlayer::IsReady() const
 {
     // To be ready, we need the AVPlayer setup
-    bool bIsReady = (MediaPlayer != nil) && ([MediaPlayer status] == AVPlayerStatusReadyToPlay);
+	if ((MediaPlayer == nil) || ([MediaPlayer status] != AVPlayerStatusReadyToPlay))
+	{
+		return false;
+	}
 
 	// A player item setup and ready to stream,
-    bIsReady &= ((MediaHelper != nil) && ([MediaHelper bIsPlayerItemReady]));
+	if ((MediaHelper == nil) || (![MediaHelper bIsPlayerItemReady]))
+	{
+		return false;
+	}
     
     // and all tracks to be setup and ready
-    for( const IMediaTrackRef& Track : Tracks )
+    for( const IMediaAudioTrackRef& AudioTrack : AudioTracks )
     {
-        FAvfMediaVideoTrack* AVFTrack = (FAvfMediaVideoTrack*)&Track.Get();
-        if( AVFTrack != nil )
-        {
-            bIsReady &= AVFTrack->IsReady();
+        FAvfMediaVideoTrack* AVFTrack = (FAvfMediaVideoTrack*)&AudioTrack.Get();
+		if ((AVFTrack == nil) || !AVFTrack->IsReady())
+		{
+			return false;
         }
     }
-    
-    return bIsReady;
+
+	for( const IMediaCaptionTrackRef& CaptionTrack : CaptionTracks )
+    {
+        FAvfMediaVideoTrack* AVFTrack = (FAvfMediaVideoTrack*)&CaptionTrack.Get();
+		if ((AVFTrack == nil) || !AVFTrack->IsReady())
+		{
+			return false;
+        }
+    }
+
+	for( const IMediaVideoTrackRef& VideoTrack : VideoTracks )
+    {
+        FAvfMediaVideoTrack* AVFTrack = (FAvfMediaVideoTrack*)&VideoTrack.Get();
+		if ((AVFTrack == nil) || !AVFTrack->IsReady())
+		{
+			return false;
+        }
+    }
+
+	return true;
 }
 
 
@@ -261,11 +305,11 @@ bool FAvfMediaPlayer::Open( const FString& Url )
                     NSError* Error = nil;
                     if( [[PlayerItem asset] statusOfValueForKey:@"tracks" error:&Error] == AVKeyValueStatusLoaded )
                     {
-                        NSArray* VideoTracks = [[PlayerItem asset] tracksWithMediaType:AVMediaTypeVideo];
-                        if( VideoTracks.count > 0 )
+                        NSArray* AssetVideoTracks = [[PlayerItem asset] tracksWithMediaType:AVMediaTypeVideo];
+                        if( AssetVideoTracks.count > 0 )
                         {
-                            AVAssetTrack* VideoTrack = [VideoTracks objectAtIndex: 0];
-                            Tracks.Add( MakeShareable( new FAvfMediaVideoTrack(VideoTrack) ) );
+                            AVAssetTrack* VideoTrack = [AssetVideoTracks objectAtIndex: 0];
+                            VideoTracks.Add( MakeShareable( new FAvfMediaVideoTrack(VideoTrack) ) );
 
 #if PLATFORM_MAC
                             GameThreadCall(^{
@@ -337,9 +381,9 @@ bool FAvfMediaPlayer::Seek( const FTimespan& Time )
 	[MediaPlayer seekToTime:CurrentTimeInSeconds];
 	[[MediaPlayer currentItem] seekToTime:CurrentTimeInSeconds];
 
-    for( IMediaTrackRef& Track : Tracks )
+    for( IMediaAudioTrackRef& VideoTrack : VideoTracks )
     {
-        FAvfMediaVideoTrack* AVFTrack = (FAvfMediaVideoTrack*)&Track.Get();
+        FAvfMediaVideoTrack* AVFTrack = (FAvfMediaVideoTrack*)&VideoTrack.Get();
         if( AVFTrack != nil )
         {
             AVFTrack->SeekToTime(CurrentTimeInSeconds);
@@ -369,25 +413,26 @@ bool FAvfMediaPlayer::SetRate( float Rate )
 
 bool FAvfMediaPlayer::Tick( float DeltaTime )
 {
-    if( ShouldAdvanceFrames() )
+    if (!ShouldAdvanceFrames())
     {
-		AVPlayerItem* item = [MediaPlayer currentItem];
-	    if (item != nil)
-	    {
-	        CurrentTime = FTimespan::FromSeconds( CMTimeGetSeconds([item currentTime]) );
-	    }
+		return true;
+	}
+
+	AVPlayerItem* item = [MediaPlayer currentItem];
+	if (item != nil)
+	{
+	    CurrentTime = FTimespan::FromSeconds( CMTimeGetSeconds([item currentTime]) );
+	}
         
-        for (IMediaTrackRef& Track : Tracks)
-        {
-            FAvfMediaVideoTrack* AVFTrack = (FAvfMediaVideoTrack*)&Track.Get();
+    for (IMediaVideoTrackRef& VideoTrack : VideoTracks)
+    {
+        FAvfMediaVideoTrack* AVFTrack = (FAvfMediaVideoTrack*)&VideoTrack.Get();
         
-            if (AVFTrack != nil)
-            {
-                AVFTrack->ReadFrameAtTime([[MediaPlayer currentItem] currentTime]);
-            }
+        if (AVFTrack != nil)
+        {
+            AVFTrack->ReadFrameAtTime([[MediaPlayer currentItem] currentTime]);
         }
     }
-
     
     return true;
 }
diff --git a/Engine/Plugins/Media/AvfMedia/Source/AvfMedia/Private/Player/AvfMediaPlayer.h b/Engine/Plugins/Media/AvfMedia/Source/AvfMedia/Private/Player/AvfMediaPlayer.h
--- a/Engine/Plugins/Media/AvfMedia/Source/AvfMedia/Private/Player/AvfMediaPlayer.h
+++ b/Engine/Plugins/Media/AvfMedia/Source/AvfMedia/Private/Player/AvfMediaPlayer.h
@@ -58,10 +58,12 @@ public:
 	// IMediaPlayer interface
 
 	virtual void Close() override;
+	virtual const TArray<IMediaAudioTrackRef>& GetAudioTracks() const override;
+	virtual const TArray<IMediaCaptionTrackRef>& GetCaptionTracks() const override;
 	virtual const IMediaInfo& GetMediaInfo() const override;
 	virtual float GetRate() const override;
 	virtual FTimespan GetTime() const override;
-	virtual const TArray<IMediaTrackRef>& GetTracks() const override;
+	virtual const TArray<IMediaVideoTrackRef>& GetVideoTracks() const override;
 	virtual bool IsLooping() const override;
 	virtual bool IsPaused() const override;
 	virtual bool IsPlaying() const override;
@@ -107,15 +109,18 @@ protected:
 
 private:
 
+	/** The available audio tracks. */
+	TArray<IMediaAudioTrackRef> AudioTracks;
+
+	/** The available caption tracks. */
+	TArray<IMediaCaptionTrackRef> CaptionTracks;
+
     /** The AVFoundation media player */
     AVPlayer* MediaPlayer;
 
     /** The player item which the media player uses to progress. */
     AVPlayerItem* PlayerItem;
 
-    /** The available media tracks. */
-    TArray<IMediaTrackRef> Tracks;
-
     /** The duration of the media. */
     FTimespan Duration;
 
@@ -131,6 +136,9 @@ private:
 	/** Cocoa helper object we can use to keep track of ns property changes in our media items */
 	FMediaHelper* MediaHelper;
 
+	/** The available video tracks. */
+	TArray<IMediaVideoTrackRef> VideoTracks;
+
 private:
 
 	/** Holds an event delegate that is invoked when media is being unloaded. */
diff --git a/Engine/Plugins/Media/AvfMedia/Source/AvfMedia/Private/Tracks/AvfMediaAudioTrack.h b/Engine/Plugins/Media/AvfMedia/Source/AvfMedia/Private/Tracks/AvfMediaAudioTrack.h
--- a/Engine/Plugins/Media/AvfMedia/Source/AvfMedia/Private/Tracks/AvfMediaAudioTrack.h
+++ b/Engine/Plugins/Media/AvfMedia/Source/AvfMedia/Private/Tracks/AvfMediaAudioTrack.h
@@ -2,13 +2,15 @@
 
 #pragma once
 
+#include "IMediaAudioTrack.h"
+
 
 /**
  * Audio Track implementation using the AV Foundation Framework.
  */
 class FAvfMediaAudioTrack
 	: public FAvfMediaTrack
-	, public IMediaTrackAudioDetails
+	, public IMediaAudioTrack
 {
 public:
 
@@ -22,7 +24,7 @@ public:
 
 public:
 
-	// IMediaTrackAudioDetails interface
+	// IMediaAudioTrack interface
 	
 	virtual uint32 GetNumChannels() const override
 	{
@@ -34,32 +36,11 @@ public:
 		return SamplesPerSecond;
 	}
 
-public:
-
-	// IMediaTrack interface
-
-	virtual const IMediaTrackAudioDetails& GetAudioDetails() const override
+	virtual IMediaStream& GetStream() override
 	{
 		return *this;
 	}
 
-	virtual const IMediaTrackCaptionDetails& GetCaptionDetails() const override
-	{
-		check(false); // not a caption track
-		return (IMediaTrackCaptionDetails&)*this;
-	}
-
-	virtual EMediaTrackTypes GetType() const override
-	{
-		return EMediaTrackTypes::Audio;
-	}
-
-	virtual const IMediaTrackVideoDetails& GetVideoDetails() const override
-	{
-		check(false); // not an video track
-		return (IMediaTrackVideoDetails&)*this;
-	}
-
 private:
 
 	/** The number of channels. */
diff --git a/Engine/Plugins/Media/AvfMedia/Source/AvfMedia/Private/Tracks/AvfMediaCaptionTrack.h b/Engine/Plugins/Media/AvfMedia/Source/AvfMedia/Private/Tracks/AvfMediaCaptionTrack.h
--- a/Engine/Plugins/Media/AvfMedia/Source/AvfMedia/Private/Tracks/AvfMediaCaptionTrack.h
+++ b/Engine/Plugins/Media/AvfMedia/Source/AvfMedia/Private/Tracks/AvfMediaCaptionTrack.h
@@ -2,13 +2,15 @@
 
 #pragma once
 
+#include "IMediaCaptionTrack.h"
+
 
 /**
  * Caption Track implementation using the AV Foundation Framework.
  */
 class FAvfMediaCaptionTrack
 	: public FAvfMediaTrack
-	, public IMediaTrackCaptionDetails
+	, public IMediaCaptionTrack
 {
 public:
 
@@ -19,31 +21,10 @@ public:
 
 public:
 
-	// IMediaTrackCaptionDetails interface
-
-public:
-
-	// IMediaTrack interface
-
-	virtual const IMediaTrackAudioDetails& GetAudioDetails() const override
-	{
-		check(false); // not an audio track
-		return (IMediaTrackAudioDetails&)*this;
-	}
+	// IMediaCaptionTrack interface
 
-	virtual const IMediaTrackCaptionDetails& GetCaptionDetails() const override
+	virtual IMediaStream& GetStream() override
 	{
 		return *this;
 	}
-
-	virtual EMediaTrackTypes GetType() const override
-	{
-		return EMediaTrackTypes::Caption;
-	}
-
-	virtual const IMediaTrackVideoDetails& GetVideoDetails() const override
-	{
-		check(false); // not a video track
-		return (IMediaTrackVideoDetails&)*this;
-	}
 };
diff --git a/Engine/Plugins/Media/AvfMedia/Source/AvfMedia/Private/Tracks/AvfMediaTrack.cpp b/Engine/Plugins/Media/AvfMedia/Source/AvfMedia/Private/Tracks/AvfMediaTrack.cpp
--- a/Engine/Plugins/Media/AvfMedia/Source/AvfMedia/Private/Tracks/AvfMediaTrack.cpp
+++ b/Engine/Plugins/Media/AvfMedia/Source/AvfMedia/Private/Tracks/AvfMediaTrack.cpp
@@ -27,7 +27,7 @@ FText FAvfMediaTrack::GetDisplayName() const
 
 	if (Name.IsEmpty())
 	{
-		DisplayName = FText::Format(LOCTEXT("UnnamedTrackFormat", "Unnamed Track {0}"), FText::AsNumber((uint32)StreamIndex));
+		DisplayName = FText::Format(LOCTEXT("UnnamedTrackFormat", "Unnamed Stream {0}"), FText::AsNumber((uint32)StreamIndex));
 	}
 	else
 	{
@@ -43,12 +43,6 @@ FText FAvfMediaTrack::GetDisplayName() const
 }
 
 
-uint32 FAvfMediaTrack::GetIndex() const
-{
-	return StreamIndex;
-}
-
-
 FString FAvfMediaTrack::GetLanguage() const
 {
 	return Language;
@@ -68,7 +62,7 @@ bool FAvfMediaTrack::IsEnabled() const
 }
 
 
-bool FAvfMediaTrack::IsMutuallyExclusive( const IMediaTrackRef& Other ) const
+bool FAvfMediaTrack::IsMutuallyExclusive( const IMediaStreamRef& Other ) const
 {
 	return false;
 }
diff --git a/Engine/Plugins/Media/AvfMedia/Source/AvfMedia/Private/Tracks/AvfMediaTrack.h b/Engine/Plugins/Media/AvfMedia/Source/AvfMedia/Private/Tracks/AvfMediaTrack.h
--- a/Engine/Plugins/Media/AvfMedia/Source/AvfMedia/Private/Tracks/AvfMediaTrack.h
+++ b/Engine/Plugins/Media/AvfMedia/Source/AvfMedia/Private/Tracks/AvfMediaTrack.h
@@ -7,7 +7,7 @@
  * Abstract base class for media tracks using the AV Foundation Framework.
  */
 class FAvfMediaTrack
-	: public IMediaTrack
+	: public IMediaStream
 {
 public:
 
@@ -21,17 +21,16 @@ public:
 
 	// IMediaTrack interface
 
-	virtual void AddSink( const IMediaSinkRef& Sink ) = 0;
+	virtual void AddSink(const IMediaSinkRef& Sink) = 0;
 	virtual bool Disable() override;
 	virtual bool Enable() override;
 	virtual FText GetDisplayName() const override;
-	virtual uint32 GetIndex() const override;
 	virtual FString GetLanguage() const override;
 	virtual FString GetName() const override;
 	virtual bool IsEnabled() const override;
-	virtual bool IsMutuallyExclusive( const IMediaTrackRef& Other ) const override;
+	virtual bool IsMutuallyExclusive(const IMediaStreamRef& Other) const override;
 	virtual bool IsProtected() const override;
-	virtual void RemoveSink( const IMediaSinkRef& Sink ) = 0;
+	virtual void RemoveSink(const IMediaSinkRef& Sink) = 0;
     
 private:
 
diff --git a/Engine/Plugins/Media/AvfMedia/Source/AvfMedia/Private/Tracks/AvfMediaVideoTrack.h b/Engine/Plugins/Media/AvfMedia/Source/AvfMedia/Private/Tracks/AvfMediaVideoTrack.h
--- a/Engine/Plugins/Media/AvfMedia/Source/AvfMedia/Private/Tracks/AvfMediaVideoTrack.h
+++ b/Engine/Plugins/Media/AvfMedia/Source/AvfMedia/Private/Tracks/AvfMediaVideoTrack.h
@@ -2,13 +2,15 @@
 
 #pragma once
 
+#include "IMediaVideoTrack.h"
+
 
 /**
  * Video Track implementation using the AV Foundation Framework.
  */
 class FAvfMediaVideoTrack
 	: public FAvfMediaTrack
-	, public IMediaTrackVideoDetails
+	, public IMediaVideoTrack
 {
 public:
 
@@ -41,31 +43,14 @@ public:
 		return FrameRate;
 	}
 
-public:
-
-	// IMediaTrack interface
-
-	virtual const IMediaTrackAudioDetails& GetAudioDetails() const override
-	{
-		check(false); // not an audio track
-		return (IMediaTrackAudioDetails&)*this;
-	}
-
-	virtual const IMediaTrackCaptionDetails& GetCaptionDetails() const override
-	{
-		check(false); // not a caption track
-		return (IMediaTrackCaptionDetails&)*this;
-	}
-
-	virtual const IMediaTrackVideoDetails& GetVideoDetails() const override
+	virtual IMediaStream& GetStream() override
 	{
 		return *this;
 	}
 
-    virtual EMediaTrackTypes GetType() const override
-    {
-        return EMediaTrackTypes::Video;
-    }
+public:
+
+	// IMediaTrack interface
 
     virtual void AddSink( const IMediaSinkRef& Sink ) override
     {
diff --git a/Engine/Plugins/Media/MediaPlayerEditor/Source/MediaPlayerEditor/Private/MediaPlayerEditorPrivatePCH.h b/Engine/Plugins/Media/MediaPlayerEditor/Source/MediaPlayerEditor/Private/MediaPlayerEditorPrivatePCH.h
--- a/Engine/Plugins/Media/MediaPlayerEditor/Source/MediaPlayerEditor/Private/MediaPlayerEditorPrivatePCH.h
+++ b/Engine/Plugins/Media/MediaPlayerEditor/Source/MediaPlayerEditor/Private/MediaPlayerEditorPrivatePCH.h
@@ -13,7 +13,10 @@
 #include "SlateStyle.h"
 #include "EditorStyle.h"
 #include "IDetailsView.h"
-#include "IMediaTrack.h"
+#include "IMediaStream.h"
+#include "IMediaAudioTrack.h"
+#include "IMediaCaptionTrack.h"
+#include "IMediaVideoTrack.h"
 #include "MediaPlayer.h"
 #include "MediaSampleBuffer.h"
 #include "MediaSoundWave.h"
diff --git a/Engine/Plugins/Media/MediaPlayerEditor/Source/MediaPlayerEditor/Private/Models/MediaPlayerEditorTexture.h b/Engine/Plugins/Media/MediaPlayerEditor/Source/MediaPlayerEditor/Private/Models/MediaPlayerEditorTexture.h
--- a/Engine/Plugins/Media/MediaPlayerEditor/Source/MediaPlayerEditor/Private/Models/MediaPlayerEditorTexture.h
+++ b/Engine/Plugins/Media/MediaPlayerEditor/Source/MediaPlayerEditor/Private/Models/MediaPlayerEditorTexture.h
@@ -14,20 +14,20 @@ public:
 	 * @param InSlateTexture The Slate texture to render to.
 	 * @param InVideoTrack The video track to fetch.
 	 */
-	FMediaPlayerEditorTexture( FSlateTexture2DRHIRef* InSlateTexture, const IMediaTrackRef& InVideoTrack )
+	FMediaPlayerEditorTexture( FSlateTexture2DRHIRef* InSlateTexture, const IMediaVideoTrackRef& InVideoTrack )
 		: FTickableObjectRenderThread(false)
 		, LastFrameTime(FTimespan::MinValue())
 		, SlateTexture(InSlateTexture)
 		, VideoBuffer(MakeShareable(new FMediaSampleBuffer))
 		, VideoTrack(InVideoTrack)
 	{
-		VideoTrack->AddSink(VideoBuffer);
+		VideoTrack->GetStream().AddSink(VideoBuffer);
 	}
 
 	/** Destructor. */
 	~FMediaPlayerEditorTexture()
 	{
-		VideoTrack->RemoveSink(VideoBuffer);
+		VideoTrack->GetStream().RemoveSink(VideoBuffer);
 	}
 
 public:
@@ -92,5 +92,5 @@ private:
 	FMediaSampleBufferRef VideoBuffer;
 
 	/** Holds the selected video track. */
-	IMediaTrackRef VideoTrack;
+	IMediaVideoTrackRef VideoTrack;
 };
diff --git a/Engine/Plugins/Media/MediaPlayerEditor/Source/MediaPlayerEditor/Private/Models/MediaPlayerEditorViewport.cpp b/Engine/Plugins/Media/MediaPlayerEditor/Source/MediaPlayerEditor/Private/Models/MediaPlayerEditorViewport.cpp
--- a/Engine/Plugins/Media/MediaPlayerEditor/Source/MediaPlayerEditor/Private/Models/MediaPlayerEditorViewport.cpp
+++ b/Engine/Plugins/Media/MediaPlayerEditor/Source/MediaPlayerEditor/Private/Models/MediaPlayerEditorViewport.cpp
@@ -1,7 +1,7 @@
 // Copyright 1998-2015 Epic Games, Inc. All Rights Reserved.
 
 #include "MediaPlayerEditorPrivatePCH.h"
-#include "IMediaTrackVideoDetails.h"
+#include "IMediaVideoTrack.h"
 
 
 /* FMediaPlayerEditorViewport structors
@@ -22,16 +22,14 @@ FMediaPlayerEditorViewport::~FMediaPlayerEditorViewport()
 /* FMediaPlayerEditorViewport interface
 *****************************************************************************/
 
-void FMediaPlayerEditorViewport::Initialize(const IMediaTrackPtr& VideoTrack)
+void FMediaPlayerEditorViewport::Initialize(const IMediaVideoTrackPtr& VideoTrack)
 {
 	ReleaseResources();
 
 	if (VideoTrack.IsValid())
 	{
-		check(VideoTrack->GetType() == EMediaTrackTypes::Video);
-
 		const bool bCreateEmptyTexture = true;
-		const FIntPoint VideoDimensions = VideoTrack->GetVideoDetails().GetDimensions();
+		const FIntPoint VideoDimensions = VideoTrack->GetDimensions();
 
 		SlateTexture = new FSlateTexture2DRHIRef(VideoDimensions.X, VideoDimensions.Y, PF_B8G8R8A8, nullptr, TexCreate_Dynamic | TexCreate_RenderTargetable, bCreateEmptyTexture);
 		EditorTexture = new FMediaPlayerEditorTexture(SlateTexture, VideoTrack.ToSharedRef());
diff --git a/Engine/Plugins/Media/MediaPlayerEditor/Source/MediaPlayerEditor/Private/Models/MediaPlayerEditorViewport.h b/Engine/Plugins/Media/MediaPlayerEditor/Source/MediaPlayerEditor/Private/Models/MediaPlayerEditorViewport.h
--- a/Engine/Plugins/Media/MediaPlayerEditor/Source/MediaPlayerEditor/Private/Models/MediaPlayerEditorViewport.h
+++ b/Engine/Plugins/Media/MediaPlayerEditor/Source/MediaPlayerEditor/Private/Models/MediaPlayerEditorViewport.h
@@ -21,7 +21,7 @@ public:
 	 *
 	 * @param VideoTrack The video track to render to the viewport.
 	 */
-	void Initialize(const IMediaTrackPtr& VideoTrack);
+	void Initialize(const IMediaVideoTrackPtr& VideoTrack);
 
 public:
 
diff --git a/Engine/Plugins/Media/MediaPlayerEditor/Source/MediaPlayerEditor/Private/Widgets/SMediaPlayerEditorViewer.cpp b/Engine/Plugins/Media/MediaPlayerEditor/Source/MediaPlayerEditor/Private/Widgets/SMediaPlayerEditorViewer.cpp
--- a/Engine/Plugins/Media/MediaPlayerEditor/Source/MediaPlayerEditor/Private/Widgets/SMediaPlayerEditorViewer.cpp
+++ b/Engine/Plugins/Media/MediaPlayerEditor/Source/MediaPlayerEditor/Private/Widgets/SMediaPlayerEditorViewer.cpp
@@ -29,7 +29,7 @@ SMediaPlayerEditorViewer::~SMediaPlayerEditorViewer()
 /* SMediaPlayerEditorPlayer interface
  *****************************************************************************/
 
-void SMediaPlayerEditorViewer::Construct( const FArguments& InArgs, UMediaPlayer* InMediaPlayer, const TSharedRef<ISlateStyle>& InStyle )
+void SMediaPlayerEditorViewer::Construct(const FArguments& InArgs, UMediaPlayer* InMediaPlayer, const TSharedRef<ISlateStyle>& InStyle)
 {
 	MediaPlayer = InMediaPlayer;
 	Viewport = MakeShareable(new FMediaPlayerEditorViewport());
@@ -92,7 +92,7 @@ void SMediaPlayerEditorViewer::Construct( const FArguments& InArgs, UMediaPlayer
 									.Padding(4.0f, 0.0f)
 									.VAlign(VAlign_Center)
 									[
-										SAssignNew(VideoTrackComboBox, SComboBox<IMediaTrackPtr>)
+										SAssignNew(VideoTrackComboBox, SComboBox<IMediaVideoTrackPtr>)
 											.OnGenerateWidget(this, &SMediaPlayerEditorViewer::HandleVideoTrackComboBoxGenerateWidget)
 											.OnSelectionChanged(this, &SMediaPlayerEditorViewer::HandleVideoTrackComboBoxSelectionChanged)
 											.OptionsSource(&VideoTracks)
@@ -124,7 +124,7 @@ void SMediaPlayerEditorViewer::Construct( const FArguments& InArgs, UMediaPlayer
 									.Padding(4.0f, 0.0f, 0.0f, 0.0f)
 									.VAlign(VAlign_Center)
 									[
-										SAssignNew(AudioTrackComboBox, SComboBox<IMediaTrackPtr>)
+										SAssignNew(AudioTrackComboBox, SComboBox<IMediaAudioTrackPtr>)
 											.OnGenerateWidget(this, &SMediaPlayerEditorViewer::HandleAudioTrackComboBoxGenerateWidget)
 											.OnSelectionChanged(this, &SMediaPlayerEditorViewer::HandleAudioTrackComboBoxSelectionChanged)
 											.OptionsSource(&AudioTracks)
@@ -156,7 +156,7 @@ void SMediaPlayerEditorViewer::Construct( const FArguments& InArgs, UMediaPlayer
 									.Padding(4.0f, 0.0f, 0.0f, 0.0f)
 									.VAlign(VAlign_Center)
 									[
-										SAssignNew(CaptionTrackComboBox, SComboBox<IMediaTrackPtr>)
+										SAssignNew(CaptionTrackComboBox, SComboBox<IMediaCaptionTrackPtr>)
 											.OnGenerateWidget(this, &SMediaPlayerEditorViewer::HandleCaptionTrackComboBoxGenerateWidget)
 											.OnSelectionChanged(this, &SMediaPlayerEditorViewer::HandleCaptionTrackComboBoxSelectionChanged)
 											.OptionsSource(&CaptionTracks)
@@ -295,11 +295,11 @@ void SMediaPlayerEditorViewer::Construct( const FArguments& InArgs, UMediaPlayer
 void SMediaPlayerEditorViewer::ReloadTracks()
 {
 	// clear track collections
-	IMediaTrackPtr SelectedCaptionTrack = CaptionTrackComboBox->GetSelectedItem();
+	IMediaCaptionTrackPtr SelectedCaptionTrack = CaptionTrackComboBox->GetSelectedItem();
 
 	if (SelectedCaptionTrack.IsValid())
 	{
-		SelectedCaptionTrack->RemoveSink(CaptionBuffer);
+		SelectedCaptionTrack->GetStream().RemoveSink(CaptionBuffer);
 	}
 
 	AudioTracks.Reset();
@@ -313,22 +313,19 @@ void SMediaPlayerEditorViewer::ReloadTracks()
 
 		if (Player.IsValid())
 		{
-			for (IMediaTrackPtr Track : Player->GetTracks())
+			for (auto AudioTrack : Player->GetAudioTracks())
 			{
-				switch (Track->GetType())
-				{
-				case EMediaTrackTypes::Audio:
-					AudioTracks.Add(Track);
-					break;
-
-				case EMediaTrackTypes::Caption:
-					CaptionTracks.Add(Track);
-					break;
-
-				case EMediaTrackTypes::Video:
-					VideoTracks.Add(Track);
-					break;
-				}
+				AudioTracks.Add(AudioTrack);
+			}
+
+			for (auto CaptionTrack : Player->GetCaptionTracks())
+			{
+				CaptionTracks.Add(CaptionTrack);
+			}
+
+			for (auto VideoTrack : Player->GetVideoTracks())
+			{
+				VideoTracks.Add(VideoTrack);
 			}
 		}
 	}
@@ -351,7 +348,7 @@ void SMediaPlayerEditorViewer::ReloadTracks()
 		if (CaptionTracks.Num() > 0)
 		{
 			CaptionTrackComboBox->SetSelectedItem(CaptionTracks[0]);
-			CaptionTracks[0]->AddSink(CaptionBuffer);
+			CaptionTracks[0]->GetStream().AddSink(CaptionBuffer);
 		}
 		else
 		{
@@ -389,14 +386,14 @@ EActiveTimerReturnType SMediaPlayerEditorViewer::HandleActiveTimer(double InCurr
 }
 
 
-TSharedRef<SWidget> SMediaPlayerEditorViewer::HandleAudioTrackComboBoxGenerateWidget( IMediaTrackPtr Value ) const
+TSharedRef<SWidget> SMediaPlayerEditorViewer::HandleAudioTrackComboBoxGenerateWidget(IMediaAudioTrackPtr Value) const
 {
 	return SNew(STextBlock)
-		.Text(Value->GetDisplayName());
+		.Text(Value->GetStream().GetDisplayName());
 }
 
 
-void SMediaPlayerEditorViewer::HandleAudioTrackComboBoxSelectionChanged( IMediaTrackPtr Selection, ESelectInfo::Type SelectInfo )
+void SMediaPlayerEditorViewer::HandleAudioTrackComboBoxSelectionChanged(IMediaAudioTrackPtr Selection, ESelectInfo::Type SelectInfo)
 {
 	if (SelectInfo != ESelectInfo::Direct)
 	{
@@ -407,11 +404,11 @@ void SMediaPlayerEditorViewer::HandleAudioTrackComboBoxSelectionChanged( IMediaT
 
 FText SMediaPlayerEditorViewer::HandleAudioTrackComboBoxText() const
 {
-	IMediaTrackPtr AudioTrack = AudioTrackComboBox->GetSelectedItem();
+	IMediaAudioTrackPtr AudioTrack = AudioTrackComboBox->GetSelectedItem();
 
 	if (AudioTrack.IsValid())
 	{
-		return AudioTrack->GetDisplayName();
+		return AudioTrack->GetStream().GetDisplayName();
 	}
 
 	if (AudioTracks.Num() == 0)
@@ -423,14 +420,14 @@ FText SMediaPlayerEditorViewer::HandleAudioTrackComboBoxText() const
 }
 
 
-TSharedRef<SWidget> SMediaPlayerEditorViewer::HandleCaptionTrackComboBoxGenerateWidget( IMediaTrackPtr Value ) const
+TSharedRef<SWidget> SMediaPlayerEditorViewer::HandleCaptionTrackComboBoxGenerateWidget(IMediaCaptionTrackPtr Value) const
 {
 	return SNew(STextBlock)
-		.Text(Value->GetDisplayName());
+		.Text(Value->GetStream().GetDisplayName());
 }
 
 
-void SMediaPlayerEditorViewer::HandleCaptionTrackComboBoxSelectionChanged( IMediaTrackPtr Selection, ESelectInfo::Type SelectInfo )
+void SMediaPlayerEditorViewer::HandleCaptionTrackComboBoxSelectionChanged(IMediaCaptionTrackPtr Selection, ESelectInfo::Type SelectInfo)
 {
 	if (SelectInfo != ESelectInfo::Direct)
 	{
@@ -441,11 +438,11 @@ void SMediaPlayerEditorViewer::HandleCaptionTrackComboBoxSelectionChanged( IMedi
 
 FText SMediaPlayerEditorViewer::HandleCaptionTrackComboBoxText() const
 {
-	IMediaTrackPtr CaptionTrack = CaptionTrackComboBox->GetSelectedItem();
+	IMediaCaptionTrackPtr CaptionTrack = CaptionTrackComboBox->GetSelectedItem();
 
 	if (CaptionTrack.IsValid())
 	{
-		return CaptionTrack->GetDisplayName();
+		return CaptionTrack->GetStream().GetDisplayName();
 	}
 
 	if (CaptionTracks.Num() == 0)
@@ -629,14 +626,14 @@ EVisibility SMediaPlayerEditorViewer::HandleTrackSelectionBoxVisibility() const
 }
 
 
-TSharedRef<SWidget> SMediaPlayerEditorViewer::HandleVideoTrackComboBoxGenerateWidget( IMediaTrackPtr Value ) const
+TSharedRef<SWidget> SMediaPlayerEditorViewer::HandleVideoTrackComboBoxGenerateWidget(IMediaVideoTrackPtr Value) const
 {
 	return SNew(STextBlock)
-		.Text(Value->GetDisplayName());
+		.Text(Value->GetStream().GetDisplayName());
 }
 
 
-void SMediaPlayerEditorViewer::HandleVideoTrackComboBoxSelectionChanged( IMediaTrackPtr Selection, ESelectInfo::Type SelectInfo )
+void SMediaPlayerEditorViewer::HandleVideoTrackComboBoxSelectionChanged(IMediaVideoTrackPtr Selection, ESelectInfo::Type SelectInfo)
 {
 	if (SelectInfo != ESelectInfo::Direct)
 	{
@@ -647,11 +644,11 @@ void SMediaPlayerEditorViewer::HandleVideoTrackComboBoxSelectionChanged( IMediaT
 
 FText SMediaPlayerEditorViewer::HandleVideoTrackComboBoxText() const
 {
-	IMediaTrackPtr VideoTrack = VideoTrackComboBox->GetSelectedItem();
+	IMediaVideoTrackPtr VideoTrack = VideoTrackComboBox->GetSelectedItem();
 
 	if (VideoTrack.IsValid())
 	{
-		return VideoTrack->GetDisplayName();
+		return VideoTrack->GetStream().GetDisplayName();
 	}
 
 	if (VideoTracks.Num() == 0)
diff --git a/Engine/Plugins/Media/MediaPlayerEditor/Source/MediaPlayerEditor/Private/Widgets/SMediaPlayerEditorViewer.h b/Engine/Plugins/Media/MediaPlayerEditor/Source/MediaPlayerEditor/Private/Widgets/SMediaPlayerEditorViewer.h
--- a/Engine/Plugins/Media/MediaPlayerEditor/Source/MediaPlayerEditor/Private/Widgets/SMediaPlayerEditorViewer.h
+++ b/Engine/Plugins/Media/MediaPlayerEditor/Source/MediaPlayerEditor/Private/Widgets/SMediaPlayerEditorViewer.h
@@ -31,7 +31,7 @@ public:
 	 * @param InMediaPlayer The UMediaPlayer asset to show the details for.
 	 * @param InStyleSet The style set to use.
 	 */
-	void Construct( const FArguments& InArgs, UMediaPlayer* InMediaPlayer, const TSharedRef<ISlateStyle>& InStyle );
+	void Construct(const FArguments& InArgs, UMediaPlayer* InMediaPlayer, const TSharedRef<ISlateStyle>& InStyle);
 
 protected:
 
@@ -44,19 +44,19 @@ private:
 	EActiveTimerReturnType HandleActiveTimer(double InCurrentTime, float InDeltaTime);
 
 	/** Callback for generating a widget for a value in the audio track combo box. */
-	TSharedRef<SWidget> HandleAudioTrackComboBoxGenerateWidget( IMediaTrackPtr Value ) const;
+	TSharedRef<SWidget> HandleAudioTrackComboBoxGenerateWidget(IMediaAudioTrackPtr Value) const;
 
 	/** Callback for changing the selected audio track. */
-	void HandleAudioTrackComboBoxSelectionChanged( IMediaTrackPtr Selection, ESelectInfo::Type SelectInfo );
+	void HandleAudioTrackComboBoxSelectionChanged(IMediaAudioTrackPtr Selection, ESelectInfo::Type SelectInfo);
 
 	/** Callback for getting the text in the audio track combo box. */
 	FText HandleAudioTrackComboBoxText() const;
 
 	/** Callback for generating a widget for a value in the caption track combo box. */
-	TSharedRef<SWidget> HandleCaptionTrackComboBoxGenerateWidget( IMediaTrackPtr Value ) const;
+	TSharedRef<SWidget> HandleCaptionTrackComboBoxGenerateWidget(IMediaCaptionTrackPtr Value) const;
 
 	/** Callback for changing the selected audio track. */
-	void HandleCaptionTrackComboBoxSelectionChanged( IMediaTrackPtr Selection, ESelectInfo::Type SelectInfo );
+	void HandleCaptionTrackComboBoxSelectionChanged(IMediaCaptionTrackPtr Selection, ESelectInfo::Type SelectInfo);
 
 	/** Callback for getting the text in the caption track combo box. */
 	FText HandleCaptionTrackComboBoxText() const;
@@ -113,10 +113,10 @@ private:
 	EVisibility HandleTrackSelectionBoxVisibility() const;
 
 	/** Callback for generating a widget for a value in the video track combo box. */
-	TSharedRef<SWidget> HandleVideoTrackComboBoxGenerateWidget( IMediaTrackPtr Value ) const;
+	TSharedRef<SWidget> HandleVideoTrackComboBoxGenerateWidget(IMediaVideoTrackPtr Value) const;
 
 	/** Callback for changing the selected video track. */
-	void HandleVideoTrackComboBoxSelectionChanged( IMediaTrackPtr Selection, ESelectInfo::Type SelectInfo );
+	void HandleVideoTrackComboBoxSelectionChanged(IMediaVideoTrackPtr Selection, ESelectInfo::Type SelectInfo);
 
 	/** Callback for getting the text in the video track combo box. */
 	FText HandleVideoTrackComboBoxText() const;
@@ -124,19 +124,19 @@ private:
 private:
 
 	/** The collection of available audio tracks. */
-	TArray<IMediaTrackPtr> AudioTracks;
+	TArray<IMediaAudioTrackPtr> AudioTracks;
 
 	/** Holds the audio track combo box. */
-	TSharedPtr<SComboBox<IMediaTrackPtr>> AudioTrackComboBox;
+	TSharedPtr<SComboBox<IMediaAudioTrackPtr>> AudioTrackComboBox;
 
 	/** Buffers text from the selected caption track. */
 	FMediaSampleBufferRef CaptionBuffer;
 
 	/** The collection of available caption tracks. */
-	TArray<IMediaTrackPtr> CaptionTracks;
+	TArray<IMediaCaptionTrackPtr> CaptionTracks;
 
 	/** Holds the caption track combo box. */
-	TSharedPtr<SComboBox<IMediaTrackPtr>> CaptionTrackComboBox;
+	TSharedPtr<SComboBox<IMediaCaptionTrackPtr>> CaptionTrackComboBox;
 
 	/** Pointer to the media player that is being viewed. */
 	UMediaPlayer* MediaPlayer;
@@ -148,11 +148,11 @@ private:
 	TSharedPtr<SSlider> ScrubberSlider;
 
 	/** The collection of available video tracks. */
-	TArray<IMediaTrackPtr> VideoTracks;
+	TArray<IMediaVideoTrackPtr> VideoTracks;
 
 	/** Holds the video track combo box. */
-	TSharedPtr<SComboBox<IMediaTrackPtr>> VideoTrackComboBox;
+	TSharedPtr<SComboBox<IMediaVideoTrackPtr>> VideoTrackComboBox;
 
-	/** Holds the Slate viewport. */
+	/** Holds the Slate view port. */
 	TSharedPtr<FMediaPlayerEditorViewport> Viewport;
 };
diff --git a/Engine/Plugins/Media/WmfMedia/Source/WmfMedia/Private/Player/WmfMediaPlayer.cpp b/Engine/Plugins/Media/WmfMedia/Source/WmfMedia/Private/Player/WmfMediaPlayer.cpp
--- a/Engine/Plugins/Media/WmfMedia/Source/WmfMedia/Private/Player/WmfMediaPlayer.cpp
+++ b/Engine/Plugins/Media/WmfMedia/Source/WmfMedia/Private/Player/WmfMediaPlayer.cpp
@@ -86,7 +86,9 @@ void FWmfMediaPlayer::Close()
 		MediaSource = NULL;
 	}
 
-	Tracks.Reset();
+	AudioTracks.Reset();
+	CaptionTracks.Reset();
+	VideoTracks.Reset();
 
 	Duration = 0;
 	MediaUrl = FString();
@@ -96,6 +98,18 @@ void FWmfMediaPlayer::Close()
 }
 
 
+const TArray<IMediaAudioTrackRef>& FWmfMediaPlayer::GetAudioTracks() const
+{
+	return AudioTracks;
+}
+
+
+const TArray<IMediaCaptionTrackRef>& FWmfMediaPlayer::GetCaptionTracks() const
+{
+	return CaptionTracks;
+}
+
+
 const IMediaInfo& FWmfMediaPlayer::GetMediaInfo() const 
 {
 	return *this;
@@ -118,9 +132,9 @@ FTimespan FWmfMediaPlayer::GetTime() const
 }
 
 
-const TArray<IMediaTrackRef>& FWmfMediaPlayer::GetTracks() const
+const TArray<IMediaVideoTrackRef>& FWmfMediaPlayer::GetVideoTracks() const
 {
-	return Tracks;
+	return VideoTracks;
 }
 
 
@@ -388,15 +402,15 @@ void FWmfMediaPlayer::AddStreamToTopology(uint32 StreamIndex, IMFTopology* Topol
 	// create and add track
 	if (MajorType == MFMediaType_Audio)
 	{
-		Tracks.Add(MakeShareable(new FWmfMediaAudioTrack(OutputType, PresentationDescriptor, Sampler, StreamDescriptor, StreamIndex)));
+		AudioTracks.Add(MakeShareable(new FWmfMediaAudioTrack(OutputType, PresentationDescriptor, Sampler, StreamDescriptor, StreamIndex)));
 	}
 	else if (MajorType == MFMediaType_SAMI)
 	{
-		Tracks.Add(MakeShareable(new FWmfMediaCaptionTrack(OutputType, PresentationDescriptor, Sampler, StreamDescriptor, StreamIndex)));
+		CaptionTracks.Add(MakeShareable(new FWmfMediaCaptionTrack(OutputType, PresentationDescriptor, Sampler, StreamDescriptor, StreamIndex)));
 	}
 	else if (MajorType == MFMediaType_Video)
 	{
-		Tracks.Add(MakeShareable(new FWmfMediaVideoTrack(OutputType, PresentationDescriptor, Sampler, StreamDescriptor, StreamIndex)));
+		VideoTracks.Add(MakeShareable(new FWmfMediaVideoTrack(OutputType, PresentationDescriptor, Sampler, StreamDescriptor, StreamIndex)));
 	}
 }
 
@@ -455,7 +469,7 @@ bool FWmfMediaPlayer::InitializeMediaSession(IUnknown* SourceObject, const FStri
 		AddStreamToTopology(StreamIndex, Topology, PresentationDescriptor, MediaSourceObject);
 	}
 
-	UE_LOG(LogWmfMedia, Verbose, TEXT("Added a total of %i tracks"), Tracks.Num());
+	UE_LOG(LogWmfMedia, Verbose, TEXT("Added a total of %i audio tracks, %i caption tracks, %i video tracks"), AudioTracks.Num(), CaptionTracks.Num(), VideoTracks.Num());
 	TracksChangedEvent.Broadcast();
 
 	UINT64 PresentationDuration = 0;
diff --git a/Engine/Plugins/Media/WmfMedia/Source/WmfMedia/Private/Player/WmfMediaPlayer.h b/Engine/Plugins/Media/WmfMedia/Source/WmfMedia/Private/Player/WmfMediaPlayer.h
--- a/Engine/Plugins/Media/WmfMedia/Source/WmfMedia/Private/Player/WmfMediaPlayer.h
+++ b/Engine/Plugins/Media/WmfMedia/Source/WmfMedia/Private/Player/WmfMediaPlayer.h
@@ -36,10 +36,12 @@ public:
 	// IMediaPlayer interface
 
 	virtual void Close() override;
+	virtual const TArray<IMediaAudioTrackRef>& GetAudioTracks() const override;
+	virtual const TArray<IMediaCaptionTrackRef>& GetCaptionTracks() const override;
 	virtual const IMediaInfo& GetMediaInfo() const override;
 	virtual float GetRate() const override;
 	virtual FTimespan GetTime() const override;
-	virtual const TArray<IMediaTrackRef>& GetTracks() const override;
+	virtual const TArray<IMediaVideoTrackRef>& GetVideoTracks() const override;
 	virtual bool IsLooping() const override;
 	virtual bool IsPaused() const override;
 	virtual bool IsPlaying() const override;
@@ -97,8 +99,11 @@ private:
 
 private:
 
-	/** The available media tracks. */
-	TArray<IMediaTrackRef> Tracks;
+	/** The available audio tracks. */
+	TArray<IMediaAudioTrackRef> AudioTracks;
+
+	/** The available caption tracks. */
+	TArray<IMediaCaptionTrackRef> CaptionTracks;
 
 	/** The duration of the currently loaded media. */
 	FTimespan Duration;
@@ -112,6 +117,9 @@ private:
 	/** The URL of the currently opened media. */
 	FString MediaUrl;
 
+	/** The available video tracks. */
+	TArray<IMediaVideoTrackRef> VideoTracks;
+
 private:
 
 	/** Holds an event delegate that is invoked when media has been closed. */
diff --git a/Engine/Plugins/Media/WmfMedia/Source/WmfMedia/Private/Tracks/WmfMediaAudioTrack.h b/Engine/Plugins/Media/WmfMedia/Source/WmfMedia/Private/Tracks/WmfMediaAudioTrack.h
--- a/Engine/Plugins/Media/WmfMedia/Source/WmfMedia/Private/Tracks/WmfMediaAudioTrack.h
+++ b/Engine/Plugins/Media/WmfMedia/Source/WmfMedia/Private/Tracks/WmfMediaAudioTrack.h
@@ -7,7 +7,7 @@
 
 class FWmfMediaAudioTrack
 	: public FWmfMediaTrack
-	, public IMediaTrackAudioDetails
+	, public IMediaAudioTrack
 {
 public:
 
@@ -29,7 +29,7 @@ public:
 
 public:
 
-	// IMediaTrackAudioDetails interface
+	// IMediaAudioTrack interface
 	
 	virtual uint32 GetNumChannels() const override
 	{
@@ -41,32 +41,11 @@ public:
 		return SamplesPerSecond;
 	}
 
-public:
-
-	// IMediaTrack interface
-
-	virtual const IMediaTrackAudioDetails& GetAudioDetails() const override
+	virtual IMediaStream& GetStream() override
 	{
 		return *this;
 	}
 
-	virtual const IMediaTrackCaptionDetails& GetCaptionDetails() const override
-	{
-		check(false); // not a caption track
-		return (IMediaTrackCaptionDetails&)*this;
-	}
-
-	virtual EMediaTrackTypes GetType() const override
-	{
-		return EMediaTrackTypes::Audio;
-	}
-
-	virtual const IMediaTrackVideoDetails& GetVideoDetails() const override
-	{
-		check(false); // not an video track
-		return (IMediaTrackVideoDetails&)*this;
-	}
-
 private:
 
 	/** The number of channels. */
diff --git a/Engine/Plugins/Media/WmfMedia/Source/WmfMedia/Private/Tracks/WmfMediaCaptionTrack.h b/Engine/Plugins/Media/WmfMedia/Source/WmfMedia/Private/Tracks/WmfMediaCaptionTrack.h
--- a/Engine/Plugins/Media/WmfMedia/Source/WmfMedia/Private/Tracks/WmfMediaCaptionTrack.h
+++ b/Engine/Plugins/Media/WmfMedia/Source/WmfMedia/Private/Tracks/WmfMediaCaptionTrack.h
@@ -7,7 +7,7 @@
 
 class FWmfMediaCaptionTrack
 	: public FWmfMediaTrack
-	, public IMediaTrackCaptionDetails
+	, public IMediaCaptionTrack
 {
 public:
 
@@ -26,33 +26,12 @@ public:
 
 public:
 
-	// IMediaTrackCaptionDetails interface
+	// IMediaCaptionTrack interface
 
-public:
-
-	// IMediaTrack interface
-
-	virtual const IMediaTrackAudioDetails& GetAudioDetails() const override
-	{
-		check(false); // not an audio track
-		return (IMediaTrackAudioDetails&)*this;
-	}
-
-	virtual const IMediaTrackCaptionDetails& GetCaptionDetails() const override
+	virtual IMediaStream& GetStream() override
 	{
 		return *this;
 	}
-
-	virtual EMediaTrackTypes GetType() const override
-	{
-		return EMediaTrackTypes::Caption;
-	}
-
-	virtual const IMediaTrackVideoDetails& GetVideoDetails() const override
-	{
-		check(false); // not a video track
-		return (IMediaTrackVideoDetails&)*this;
-	}
 };
 
 
diff --git a/Engine/Plugins/Media/WmfMedia/Source/WmfMedia/Private/Tracks/WmfMediaTrack.cpp b/Engine/Plugins/Media/WmfMedia/Source/WmfMedia/Private/Tracks/WmfMediaTrack.cpp
--- a/Engine/Plugins/Media/WmfMedia/Source/WmfMedia/Private/Tracks/WmfMediaTrack.cpp
+++ b/Engine/Plugins/Media/WmfMedia/Source/WmfMedia/Private/Tracks/WmfMediaTrack.cpp
@@ -63,7 +63,7 @@ FText FWmfMediaTrack::GetDisplayName() const
 
 	if (Name.IsEmpty())
 	{
-		DisplayName = FText::Format(LOCTEXT("UnnamedTrackFormat", "Unnamed Track {0}"), FText::AsNumber((uint32)StreamIndex));
+		DisplayName = FText::Format(LOCTEXT("UnnamedStreamFormat", "Unnamed Stream {0}"), FText::AsNumber((uint32)StreamIndex));
 	}
 	else
 	{
@@ -79,12 +79,6 @@ FText FWmfMediaTrack::GetDisplayName() const
 }
 
 
-uint32 FWmfMediaTrack::GetIndex() const
-{
-	return StreamIndex;
-}
-
-
 FString FWmfMediaTrack::GetLanguage() const
 {
 	return Language;
@@ -106,7 +100,7 @@ bool FWmfMediaTrack::IsEnabled() const
 }
 
 
-bool FWmfMediaTrack::IsMutuallyExclusive( const IMediaTrackRef& Other ) const
+bool FWmfMediaTrack::IsMutuallyExclusive( const IMediaStreamRef& Other ) const
 {
 	return false;
 }
diff --git a/Engine/Plugins/Media/WmfMedia/Source/WmfMedia/Private/Tracks/WmfMediaTrack.h b/Engine/Plugins/Media/WmfMedia/Source/WmfMedia/Private/Tracks/WmfMediaTrack.h
--- a/Engine/Plugins/Media/WmfMedia/Source/WmfMedia/Private/Tracks/WmfMediaTrack.h
+++ b/Engine/Plugins/Media/WmfMedia/Source/WmfMedia/Private/Tracks/WmfMediaTrack.h
@@ -13,7 +13,7 @@
  * a new sample is being requested with the @see RequestNewSample method.
  */
 class FWmfMediaTrack
-	: public IMediaTrack
+	: public IMediaStream
 {
 public:
 
@@ -25,7 +25,11 @@ public:
 	 * @param InStreamDescriptor The stream's descriptor object.
 	 * @param InStreamIndex The stream's index number in the presentation.
 	 */
-	FWmfMediaTrack( IMFPresentationDescriptor* InPresentationDescriptor, FWmfMediaSampler* InSampler, IMFStreamDescriptor* InStreamDescriptor, DWORD InStreamIndex );
+	FWmfMediaTrack(
+		IMFPresentationDescriptor* InPresentationDescriptor,
+		FWmfMediaSampler* InSampler,
+		IMFStreamDescriptor* InStreamDescriptor,
+		DWORD InStreamIndex);
 
 public:
 
@@ -35,11 +39,10 @@ public:
 	virtual bool Disable() override;
 	virtual bool Enable() override;
 	virtual FText GetDisplayName() const override;
-	virtual uint32 GetIndex() const override;
 	virtual FString GetLanguage() const override;
 	virtual FString GetName() const override;
 	virtual bool IsEnabled() const override;
-	virtual bool IsMutuallyExclusive( const IMediaTrackRef& Other ) const override;
+	virtual bool IsMutuallyExclusive( const IMediaStreamRef& Other ) const override;
 	virtual bool IsProtected() const override;
 	virtual void RemoveSink( const IMediaSinkRef& Sink ) override;
 
diff --git a/Engine/Plugins/Media/WmfMedia/Source/WmfMedia/Private/Tracks/WmfMediaVideoTrack.h b/Engine/Plugins/Media/WmfMedia/Source/WmfMedia/Private/Tracks/WmfMediaVideoTrack.h
--- a/Engine/Plugins/Media/WmfMedia/Source/WmfMedia/Private/Tracks/WmfMediaVideoTrack.h
+++ b/Engine/Plugins/Media/WmfMedia/Source/WmfMedia/Private/Tracks/WmfMediaVideoTrack.h
@@ -10,7 +10,7 @@
  */
 class FWmfMediaVideoTrack
 	: public FWmfMediaTrack
-	, public IMediaTrackVideoDetails
+	, public IMediaVideoTrack
 {
 public:
 
@@ -45,7 +45,7 @@ public:
 
 public:
 
-	// IMediaTrackVideoDetails interface
+	// IMediaVideoTrack interface
 
 	virtual uint32 GetBitRate() const override
 	{
@@ -62,28 +62,7 @@ public:
 		return FrameRate;
 	}
 
-public:
-
-	// IMediaTrack interface
-
-	virtual const IMediaTrackAudioDetails& GetAudioDetails() const override
-	{
-		check(false); // not an audio track
-		return (IMediaTrackAudioDetails&)*this;
-	}
-
-	virtual const IMediaTrackCaptionDetails& GetCaptionDetails() const override
-	{
-		check(false); // not a caption track
-		return (IMediaTrackCaptionDetails&)*this;
-	}
-
-	virtual EMediaTrackTypes GetType() const override
-	{
-		return EMediaTrackTypes::Video;
-	}
-
-	virtual const IMediaTrackVideoDetails& GetVideoDetails() const override
+	virtual IMediaStream& GetStream() override
 	{
 		return *this;
 	}
diff --git a/Engine/Plugins/Media/WmfMedia/Source/WmfMedia/Private/WmfMediaPrivatePCH.h b/Engine/Plugins/Media/WmfMedia/Source/WmfMedia/Private/WmfMediaPrivatePCH.h
--- a/Engine/Plugins/Media/WmfMedia/Source/WmfMedia/Private/WmfMediaPrivatePCH.h
+++ b/Engine/Plugins/Media/WmfMedia/Source/WmfMedia/Private/WmfMediaPrivatePCH.h
@@ -11,10 +11,10 @@
 #include "IMediaInfo.h"
 #include "IMediaPlayer.h"
 #include "IMediaSink.h"
-#include "IMediaTrack.h"
-#include "IMediaTrackAudioDetails.h"
-#include "IMediaTrackCaptionDetails.h"
-#include "IMediaTrackVideoDetails.h"
+#include "IMediaStream.h"
+#include "IMediaAudioTrack.h"
+#include "IMediaCaptionTrack.h"
+#include "IMediaVideoTrack.h"
 
 #include "AllowWindowsPlatformTypes.h"
 #include <windows.h>
diff --git a/Engine/Source/Editor/DetailCustomizations/Private/MediaSoundWaveCustomization.cpp b/Engine/Source/Editor/DetailCustomizations/Private/MediaSoundWaveCustomization.cpp
--- a/Engine/Source/Editor/DetailCustomizations/Private/MediaSoundWaveCustomization.cpp
+++ b/Engine/Source/Editor/DetailCustomizations/Private/MediaSoundWaveCustomization.cpp
@@ -2,7 +2,7 @@
 
 #include "DetailCustomizationsPrivatePCH.h"
 #include "IMediaPlayer.h"
-#include "IMediaTrack.h"
+#include "IMediaStream.h"
 #include "MediaPlayer.h"
 #include "MediaSoundWaveCustomization.h"
 
diff --git a/Engine/Source/Editor/DetailCustomizations/Private/MediaTextureCustomization.cpp b/Engine/Source/Editor/DetailCustomizations/Private/MediaTextureCustomization.cpp
--- a/Engine/Source/Editor/DetailCustomizations/Private/MediaTextureCustomization.cpp
+++ b/Engine/Source/Editor/DetailCustomizations/Private/MediaTextureCustomization.cpp
@@ -2,7 +2,10 @@
 
 #include "DetailCustomizationsPrivatePCH.h"
 #include "IMediaPlayer.h"
-#include "IMediaTrack.h"
+#include "IMediaStream.h"
+#include "IMediaAudioTrack.h"
+#include "IMediaCaptionTrack.h"
+#include "IMediaVideoTrack.h"
 #include "MediaPlayer.h"
 #include "MediaTextureCustomization.h"
 
@@ -77,33 +80,30 @@ TSharedRef<SWidget> FMediaTextureCustomization::HandleVideoTrackComboButtonMenuC
 		return SNullWidget::NullWidget;
 	}
 
-	const TArray<IMediaTrackRef>& MediaTracks = Player->GetTracks();
+	const TArray<IMediaVideoTrackRef>& VideoTracks = Player->GetVideoTracks();
 
 	// populate the menu
 	FMenuBuilder MenuBuilder(true, nullptr);
 
-	for (const IMediaTrackRef& Track : MediaTracks)
+	for (int32 VideoTrackIndex = 0; VideoTrackIndex < VideoTracks.Num(); ++VideoTrackIndex)
 	{
-		if (Track->GetType() == EMediaTrackTypes::Video)
-		{
-			FUIAction Action(FExecuteAction::CreateSP(this, &FMediaTextureCustomization::HandleVideoTrackComboButtonMenuEntryExecute, Track->GetIndex()));
-
-			MenuBuilder.AddMenuEntry(
-				Track->GetDisplayName(),
-				TAttribute<FText>(),
-				FSlateIcon(),
-				Action
-			);
-		}
+		FUIAction Action(FExecuteAction::CreateSP(this, &FMediaTextureCustomization::HandleVideoTrackComboButtonMenuEntryExecute, VideoTrackIndex));
+
+		MenuBuilder.AddMenuEntry(
+			VideoTracks[VideoTrackIndex]->GetStream().GetDisplayName(),
+			TAttribute<FText>(),
+			FSlateIcon(),
+			Action
+		);
 	}
 
 	return MenuBuilder.MakeWidget();
 }
 
 
-void FMediaTextureCustomization::HandleVideoTrackComboButtonMenuEntryExecute( uint32 TrackIndex )
+void FMediaTextureCustomization::HandleVideoTrackComboButtonMenuEntryExecute( int32 TrackIndex )
 {
-	VideoTrackIndexProperty->SetValue((int32)TrackIndex);
+	VideoTrackIndexProperty->SetValue(TrackIndex);
 }
 
 
@@ -142,15 +142,15 @@ FText FMediaTextureCustomization::HandleVideoTrackComboButtonText() const
 		return LOCTEXT("NoMediaLoaded", "No media loaded");
 	}
 
-	IMediaTrackPtr Track = Player->GetTrack(VideoTrackIndex, EMediaTrackTypes::Video);
+	auto VideoTracks = Player->GetVideoTracks();
 
-	// generate track name string
-	if (Track.IsValid())
+	if (!VideoTracks.IsValidIndex(VideoTrackIndex))
 	{
-		return Track->GetDisplayName();
+		return LOCTEXT("SelectVideoTrack", "Select a Video Track...");
 	}
 
-	return LOCTEXT("SelectVideoTrack", "Select a Video Track...");
+	// get track name
+	return VideoTracks[VideoTrackIndex]->GetStream().GetDisplayName();
 }
 
 
diff --git a/Engine/Source/Editor/DetailCustomizations/Private/MediaTextureCustomization.h b/Engine/Source/Editor/DetailCustomizations/Private/MediaTextureCustomization.h
--- a/Engine/Source/Editor/DetailCustomizations/Private/MediaTextureCustomization.h
+++ b/Engine/Source/Editor/DetailCustomizations/Private/MediaTextureCustomization.h
@@ -13,7 +13,7 @@ public:
 
 	// IDetailCustomization interface
 
-	virtual void CustomizeDetails( IDetailLayoutBuilder& DetailBuilder ) override;
+	virtual void CustomizeDetails(IDetailLayoutBuilder& DetailBuilder) override;
 
 public:
 
@@ -33,7 +33,7 @@ private:
 	TSharedRef<SWidget> HandleVideoTrackComboButtonMenuContent() const;
 
 	/** Callback for selecting a track in the VideoTrack combo box. */
-	void HandleVideoTrackComboButtonMenuEntryExecute( uint32 TrackIndex );
+	void HandleVideoTrackComboButtonMenuEntryExecute(int32 TrackIndex);
 
 	/** Callback for getting the text of the VideoTrack combo box. */
 	FText HandleVideoTrackComboButtonText() const;
diff --git a/Engine/Source/Runtime/Media/Public/IMediaAudioTrack.h b/Engine/Source/Runtime/Media/Public/IMediaAudioTrack.h
new file mode 100644
--- /dev/null
+++ b/Engine/Source/Runtime/Media/Public/IMediaAudioTrack.h
@@ -0,0 +1,50 @@
+// Copyright 1998-2015 Epic Games, Inc. All Rights Reserved.
+
+#pragma once
+
+
+class IMediaStream;
+
+
+/**
+ * Interface for audio tracks.
+ */
+class IMediaAudioTrack
+{
+public:
+
+	/**
+	 * Gets the number of audio channels.
+	 *
+	 * @return Number of channels.
+	 * @see GetSamplesPerSecond
+	 */
+	virtual uint32 GetNumChannels() const = 0;
+
+	/**
+	 * Gets the number of audio samples per second.
+	 *
+	 * @return Samples per second.
+	 * @see GetNumChannels
+	 */
+	virtual uint32 GetSamplesPerSecond() const = 0;
+
+	/**
+	 * Get the underlying media stream.
+	 *
+	 * @return Media stream object.
+	 */
+	virtual IMediaStream& GetStream() = 0;
+
+public:
+
+	/** Virtual destructor. */
+	~IMediaAudioTrack() { }
+};
+
+
+/** Type definition for shared pointers to instances of IMediaAudioTrack. */
+typedef TSharedPtr<IMediaAudioTrack, ESPMode::ThreadSafe> IMediaAudioTrackPtr;
+
+/** Type definition for shared references to instances of IMediaAudioTrack. */
+typedef TSharedRef<IMediaAudioTrack, ESPMode::ThreadSafe> IMediaAudioTrackRef;
diff --git a/Engine/Source/Runtime/Media/Public/IMediaCaptionTrack.h b/Engine/Source/Runtime/Media/Public/IMediaCaptionTrack.h
new file mode 100644
--- /dev/null
+++ b/Engine/Source/Runtime/Media/Public/IMediaCaptionTrack.h
@@ -0,0 +1,34 @@
+// Copyright 1998-2015 Epic Games, Inc. All Rights Reserved.
+
+#pragma once
+
+
+class IMediaStream;
+
+
+/**
+ * Interface for caption tracks.
+ */
+class IMediaCaptionTrack
+{
+public:
+
+	/**
+	 * Get the underlying media stream.
+	 *
+	 * @return Media stream object.
+	 */
+	virtual IMediaStream& GetStream() = 0;
+
+public:
+
+	/** Virtual destructor. */
+	~IMediaCaptionTrack() { }
+};
+
+
+/** Type definition for shared pointers to instances of IMediaCaptionTrack. */
+typedef TSharedPtr<IMediaCaptionTrack, ESPMode::ThreadSafe> IMediaCaptionTrackPtr;
+
+/** Type definition for shared references to instances of IMediaCaptionTrack. */
+typedef TSharedRef<IMediaCaptionTrack, ESPMode::ThreadSafe> IMediaCaptionTrackRef;
diff --git a/Engine/Source/Runtime/Media/Public/IMediaPlayer.h b/Engine/Source/Runtime/Media/Public/IMediaPlayer.h
--- a/Engine/Source/Runtime/Media/Public/IMediaPlayer.h
+++ b/Engine/Source/Runtime/Media/Public/IMediaPlayer.h
@@ -3,10 +3,13 @@
 #pragma once
 
 #include "IMediaInfo.h"
-#include "IMediaTrack.h"
+#include "IMediaStream.h"
 
 
-class IMediaTrack;
+class IMediaStream;
+class IMediaAudioTrack;
+class IMediaCaptionTrack;
+class IMediaVideoTrack;
 
 
 /**
@@ -46,6 +49,22 @@ public:
 	 */
 	virtual void Close() = 0;
 
+	/**
+	 * Get the media's audio tracks.
+	 *
+	 * @return Collection of audio tracks.
+	 * @see GetCaptionTracks, GetVideoTracks
+	 */
+	virtual const TArray<TSharedRef<IMediaAudioTrack, ESPMode::ThreadSafe>>& GetAudioTracks() const = 0;
+
+	/**
+	 * Get the media's caption tracks.
+	 *
+	 * @return Collection of caption tracks.
+	 * @see GetAudioTracks, GetVideoTracks
+	 */
+	virtual const TArray<TSharedRef<IMediaCaptionTrack, ESPMode::ThreadSafe>>& GetCaptionTracks() const = 0;
+
 	/**
 	 * Gets the last error that occurred during media loading or playback.
 	 *
@@ -77,12 +96,12 @@ public:
 	virtual FTimespan GetTime() const = 0;
 
 	/**
-	 * Gets access to the media's audio, video and other tracks.
+	 * Get the media's video tracks.
 	 *
-	 * @return Media tracks interface.
-	 * @see GetFirstTrack, GetTrack
+	 * @return Collection of video tracks.
+	 * @see GetAudioTracks, GetCaptionTracks
 	 */
-	virtual const TArray<TSharedRef<IMediaTrack, ESPMode::ThreadSafe>>& GetTracks() const = 0;
+	virtual const TArray<TSharedRef<IMediaVideoTrack, ESPMode::ThreadSafe>>& GetVideoTracks() const = 0;
 
 	/**
 	 * Checks whether playback is currently looping.
@@ -127,7 +146,7 @@ public:
 	 * @return true if the media was opened successfully, false otherwise.
 	 * @see Close, IsReady
 	 */
-	virtual bool Open( const FString& Url ) = 0;
+	virtual bool Open(const FString& Url) = 0;
 
 	/**
 	 * Opens a media from a file or memory archive.
@@ -137,7 +156,7 @@ public:
 	 * @return true if the media was opened, false otherwise.
 	 * @see Close, IsReady
 	 */
-	virtual bool Open( const TSharedRef<FArchive, ESPMode::ThreadSafe>& Archive, const FString& OriginalUrl ) = 0;
+	virtual bool Open(const TSharedRef<FArchive, ESPMode::ThreadSafe>& Archive, const FString& OriginalUrl) = 0;
 
 	/**
 	 * Changes the media's playback time.
@@ -146,7 +165,7 @@ public:
 	 * @return true on success, false otherwise.
 	 * @see GetTime
 	 */
-	virtual bool Seek( const FTimespan& Time ) = 0;
+	virtual bool Seek(const FTimespan& Time) = 0;
 
 	/**
 	 * Sets whether playback should be looping.
@@ -154,7 +173,7 @@ public:
 	 * @param Looping Enables or disables looping.
 	 * @see IsLooping
 	 */
-	virtual bool SetLooping( bool Looping ) = 0;
+	virtual bool SetLooping(bool Looping) = 0;
 
 	/**
 	 * Sets the current playback rate.
@@ -186,49 +205,6 @@ public:
 
 public:
 
-	/**
-	 * Gets the first media track matching the specified type.
-	 *
-	 * @param TrackType The expected type of the track, i.e. audio or video.
-	 * @return The first matching track, nullptr otherwise.
-	 * @see GetTrack, GetTracks
-	 */
-	TSharedPtr<IMediaTrack, ESPMode::ThreadSafe> GetFirstTrack( EMediaTrackTypes TrackType )
-	{
-		for (const TSharedRef<IMediaTrack, ESPMode::ThreadSafe>& Track : GetTracks())
-		{
-			if (Track->GetType() == TrackType)
-			{
-				return Track;
-			}
-		}
-
-		return nullptr;
-	}
-
-	/**
-	 * Gets the media track with the specified index and type.
-	 *
-	 * @param TrackIndex The index of the track to get.
-	 * @param TrackType The expected type of the track, i.e. audio or video.
-	 * @return The track if it exists, nullptr otherwise.
-	 * @see GetFirstTrack, GetTracks
-	 */
-	TSharedPtr<IMediaTrack, ESPMode::ThreadSafe> GetTrack( int32 TrackIndex, EMediaTrackTypes TrackType )
-	{
-		if (GetTracks().IsValidIndex(TrackIndex))
-		{
-			TSharedPtr<IMediaTrack, ESPMode::ThreadSafe> Track = GetTracks()[TrackIndex];
-
-			if (Track->GetType() == TrackType)
-			{
-				return Track;
-			}
-		}
-
-		return nullptr;
-	}
-
 	/**
 	 * Pauses media playback.
 	 *
@@ -263,7 +239,7 @@ public:
 	 * @return true on success, false otherwise.
 	 * @see GetDuration, GetTime
 	 */
-	bool Seek( const FTimespan& TimeOffset, EMediaSeekDirection Direction )
+	bool Seek(const FTimespan& TimeOffset, EMediaSeekDirection Direction)
 	{
 		FTimespan SeekTime;
 
diff --git a/Engine/Source/Runtime/Media/Public/IMediaStream.h b/Engine/Source/Runtime/Media/Public/IMediaStream.h
new file mode 100644
--- /dev/null
+++ b/Engine/Source/Runtime/Media/Public/IMediaStream.h
@@ -0,0 +1,132 @@
+// Copyright 1998-2015 Epic Games, Inc. All Rights Reserved.
+
+#pragma once
+
+
+class IMediaSink;
+
+
+/**
+ * Interface for media streams.
+ *
+ * A media stream is a sequence of data of a uniform media type, such as audio or video.
+ * Media file containers may consist of a single stream or they may contain multiple
+ * streams of different types, for example video streams with different encodings or
+ * multiple audio and caption text streams for different languages.
+ *
+ * Track processing can be enabled and disabled while media playback is in progress.
+ * Use the Enable() and Disable() methods respectively. Whether a stream is enabled by
+ * default is determined by the media file. For video files the video stream as well
+ * as the audio stream for the default language are usually enabled by default.
+ *
+ * Some stream types may be mutually exclusive, which means that the media player can
+ * play only one stream of a certain type at a time. For example, if a video contains
+ * multiple audio streams, only one of them may be enabled. Use the IsMutuallyExclusive
+ * method to determine whether two streams are mutually exclusive.
+ *
+ * @see IMediaPlayer
+ */
+class IMediaStream
+{
+public:
+
+	/**
+	 * Adds the given media sink to this stream.
+	 *
+	 * @param Sink The sink to add.
+	 * @see RemoveSink
+	 */
+	virtual void AddSink(const TSharedRef<IMediaSink, ESPMode::ThreadSafe>& Sink) = 0;
+
+	/**
+	 * Disables this stream.
+	 *
+	 * If a stream is disabled, no data will be decoded for it during media playback.
+	 *
+	 * @return true on success, false otherwise.
+	 * @see Enable, IsEnabled
+	 */
+	virtual bool Disable() = 0;
+
+	/**
+	 * Disables this stream.
+	 *
+	 * If a stream is enabled, data will be decoded for it during media playback.
+	 * Enabling a stream may fail if the stream is mutually exclusive with another stream.
+	 * Use the IsMututallyExclusive method to check for mutual exclusion of streams.
+	 *
+	 * @return true on success, false otherwise.
+	 * @see Disable, IsEnabled
+	 */
+	virtual bool Enable() = 0;
+
+	/**
+	 * Generates the stream's display name.
+	 *
+	 * Unlike GetName(), which may return an empty string for unnamed streams, this method always
+	 * returns a human readable display text for the stream's name that will be generated, if needed.
+	 *
+	 * @return Display name text.
+	 * @see GetName
+	 */
+	virtual FText GetDisplayName() const = 0;
+
+	/**
+	 * Gets the stream's language tag, i.e. "en-US" for English.
+	 *
+	 * @return Language tag.
+	 */
+	virtual FString GetLanguage() const = 0;
+
+	/**
+	 * Gets the name of the stream.
+	 *
+	 * @return Track name string.
+	 * @see GetDisplayName
+	 */
+	virtual FString GetName() const = 0;
+
+	/**
+	 * Checks whether this stream is currently enabled.
+	 *
+	 * If a stream is enabled, data will be decoded for it during media playback.
+	 *
+	 * @return true if the stream is enabled, false otherwise.
+	 * @see Disable, Enable
+	 */
+	virtual bool IsEnabled() const = 0;
+
+	/**
+	 * Checks whether this stream is mutually exclusive with another.
+	 *
+	 * @param Other The other stream to check with.
+	 * @return true if the streams are mutually exclusive, false otherwise.
+	 */
+	virtual bool IsMutuallyExclusive(const TSharedRef<IMediaStream, ESPMode::ThreadSafe>& Other) const = 0;
+
+	/**
+	 * Checks whether this stream has protected content, i.e. DRM.
+	 *
+	 * @return true if the content is protected, false otherwise.
+	 */
+	virtual bool IsProtected() const = 0;
+
+	/**
+	 * Removes the given media sink from this stream.
+	 *
+	 * @param Sink The sink to remove.
+	 */
+	virtual void RemoveSink(const TSharedRef<IMediaSink, ESPMode::ThreadSafe>& Sink) = 0;
+
+public:
+
+	/** Virtual destructor. */
+	virtual ~IMediaStream() { }
+};
+
+
+/** Type definition for shared pointers to instances of IMediaTrack. */
+typedef TSharedPtr<IMediaStream, ESPMode::ThreadSafe> IMediaStreamPtr;
+
+/** Type definition for shared references to instances of IMediaTrack. */
+typedef TSharedRef<IMediaStream, ESPMode::ThreadSafe> IMediaStreamRef;
diff --git a/Engine/Source/Runtime/Media/Public/IMediaTrack.h b/Engine/Source/Runtime/Media/Public/IMediaTrack.h
deleted file mode 100644
diff --git a/Engine/Source/Runtime/Media/Public/IMediaTrackAudioDetails.h b/Engine/Source/Runtime/Media/Public/IMediaTrackAudioDetails.h
deleted file mode 100644
diff --git a/Engine/Source/Runtime/Media/Public/IMediaTrackCaptionDetails.h b/Engine/Source/Runtime/Media/Public/IMediaTrackCaptionDetails.h
deleted file mode 100644
diff --git a/Engine/Source/Runtime/Media/Public/IMediaTrackVideoDetails.h b/Engine/Source/Runtime/Media/Public/IMediaTrackVideoDetails.h
deleted file mode 100644
diff --git a/Engine/Source/Runtime/Media/Public/IMediaVideoTrack.h b/Engine/Source/Runtime/Media/Public/IMediaVideoTrack.h
new file mode 100644
--- /dev/null
+++ b/Engine/Source/Runtime/Media/Public/IMediaVideoTrack.h
@@ -0,0 +1,78 @@
+// Copyright 1998-2015 Epic Games, Inc. All Rights Reserved.
+
+#pragma once
+
+
+class IMediaStream;
+
+
+/**
+ * Interface for video tracks.
+ */
+class IMediaVideoTrack
+{
+public:
+
+	/**
+	 * Gets the average data rate of the video track in bits per second.
+	 *
+	 * @return The data rate.
+	 * @see GetFrameRate
+	 */
+	virtual uint32 GetBitRate() const = 0;
+
+	/**
+	 * Gets the video's dimensions in pixels.
+	 *
+	 * @return The width and height of the video.
+	 * @see GetAspectRatio
+	 */
+	virtual FIntPoint GetDimensions() const = 0;
+
+	/**
+	 * Gets the video's frame rate in frames per second.
+	 *
+	 * @return The frame rate.
+	 * @see GetBitRate
+	 */
+	virtual float GetFrameRate() const = 0;
+
+	/**
+	 * Get the underlying media stream.
+	 *
+	 * @return Media stream object.
+	 */
+	virtual IMediaStream& GetStream() = 0;
+
+public:
+
+	/**
+	 * Gets the media's aspect ratio.
+	 *
+	 * @return Aspect ratio.
+	 * @see GetDimensions
+	 */
+	float GetAspectRatio() const
+	{
+		const FVector2D Dimensions = GetDimensions();
+
+		if (FMath::IsNearlyZero(Dimensions.Y))
+		{
+			return 0.0f;
+		}
+
+		return Dimensions.X / Dimensions.Y;
+	}
+
+public:
+
+	/** Virtual destructor. */
+	~IMediaVideoTrack() { }
+};
+
+
+/** Type definition for shared pointers to instances of IMediaVideoTrack. */
+typedef TSharedPtr<IMediaVideoTrack, ESPMode::ThreadSafe> IMediaVideoTrackPtr;
+
+/** Type definition for shared references to instances of IMediaVideoTrack. */
+typedef TSharedRef<IMediaVideoTrack, ESPMode::ThreadSafe> IMediaVideoTrackRef;
diff --git a/Engine/Source/Runtime/MediaAssets/Private/Assets/MediaSoundWave.cpp b/Engine/Source/Runtime/MediaAssets/Private/Assets/MediaSoundWave.cpp
--- a/Engine/Source/Runtime/MediaAssets/Private/Assets/MediaSoundWave.cpp
+++ b/Engine/Source/Runtime/MediaAssets/Private/Assets/MediaSoundWave.cpp
@@ -1,7 +1,7 @@
 // Copyright 1998-2015 Epic Games, Inc. All Rights Reserved.
 
 #include "MediaAssetsPrivatePCH.h"
-#include "IMediaTrackAudioDetails.h"
+#include "IMediaAudioTrack.h"
 #include "MediaSampleQueue.h"
 #include "MediaSoundWave.h"
 
@@ -24,7 +24,7 @@ UMediaSoundWave::~UMediaSoundWave()
 {
 	if (AudioTrack.IsValid())
 	{
-		AudioTrack->RemoveSink(AudioQueue);
+		AudioTrack->GetStream().RemoveSink(AudioQueue);
 	}
 
 	if (CurrentMediaPlayer != nullptr)
@@ -167,7 +167,7 @@ void UMediaSoundWave::InitializeTrack()
 	// disconnect from current track
 	if (AudioTrack.IsValid())
 	{
-		AudioTrack->RemoveSink(AudioQueue);
+		AudioTrack->GetStream().RemoveSink(AudioQueue);
 	}
 
 	AudioTrack.Reset();
@@ -177,28 +177,26 @@ void UMediaSoundWave::InitializeTrack()
 	{
 		IMediaPlayerPtr Player = MediaPlayer->GetPlayer();
 
-		if (Player .IsValid())
+		if (Player.IsValid())
 		{
-			if (AudioTrackIndex == INDEX_NONE)
-			{
-				AudioTrack = Player->GetFirstTrack(EMediaTrackTypes::Audio);
+			auto AudioTracks = Player->GetAudioTracks();
 
-				if (AudioTrack.IsValid())
-				{
-					AudioTrackIndex = AudioTrack->GetIndex();
-				}
+			if (AudioTracks.IsValidIndex(AudioTrackIndex))
+			{
+				AudioTrack = AudioTracks[AudioTrackIndex];
 			}
-			else
+			else if (AudioTracks.Num() > 0)
 			{
-				AudioTrack = Player ->GetTrack(AudioTrackIndex, EMediaTrackTypes::Audio);
+				AudioTrack = AudioTracks[0];
+				AudioTrackIndex = 0;
 			}
 		}
 	}
 
 	if (AudioTrack.IsValid())
 	{
-		NumChannels = AudioTrack->GetAudioDetails().GetNumChannels();
-		SampleRate = AudioTrack->GetAudioDetails().GetSamplesPerSecond();
+		NumChannels = AudioTrack->GetNumChannels();
+		SampleRate = AudioTrack->GetSamplesPerSecond();
 	}
 	else
 	{
@@ -209,7 +207,7 @@ void UMediaSoundWave::InitializeTrack()
 	// connect to new track
 	if (AudioTrack.IsValid())
 	{
-		AudioTrack->AddSink(AudioQueue);
+		AudioTrack->GetStream().AddSink(AudioQueue);
 	}
 }
 
diff --git a/Engine/Source/Runtime/MediaAssets/Private/Assets/MediaTexture.cpp b/Engine/Source/Runtime/MediaAssets/Private/Assets/MediaTexture.cpp
--- a/Engine/Source/Runtime/MediaAssets/Private/Assets/MediaTexture.cpp
+++ b/Engine/Source/Runtime/MediaAssets/Private/Assets/MediaTexture.cpp
@@ -1,7 +1,7 @@
 // Copyright 1998-2015 Epic Games, Inc. All Rights Reserved.
 
 #include "MediaAssetsPrivatePCH.h"
-#include "IMediaTrackVideoDetails.h"
+#include "IMediaVideoTrack.h"
 #include "MediaSampleBuffer.h"
 
 
@@ -17,7 +17,6 @@ UMediaTexture::UMediaTexture( const FObjectInitializer& ObjectInitializer )
 	, VideoBuffer(MakeShareable(new FMediaSampleBuffer))
 {
 	NeverStream = true;
-
 	UpdateResource();
 }
 
@@ -26,7 +25,7 @@ UMediaTexture::~UMediaTexture()
 {
 	if (VideoTrack.IsValid())
 	{
-		VideoTrack->RemoveSink(VideoBuffer);
+		VideoTrack->GetStream().RemoveSink(VideoBuffer);
 	}
 
 	if (CurrentMediaPlayer != nullptr)
@@ -191,7 +190,7 @@ void UMediaTexture::InitializeTrack()
 	// disconnect from current track
 	if (VideoTrack.IsValid())
 	{
-		VideoTrack->RemoveSink(VideoBuffer);
+		VideoTrack->GetStream().RemoveSink(VideoBuffer);
 		VideoTrack.Reset();
 	}
 
@@ -202,25 +201,23 @@ void UMediaTexture::InitializeTrack()
 
 		if (Player.IsValid())
 		{
-			if (VideoTrackIndex == INDEX_NONE)
-			{
-				VideoTrack = Player->GetFirstTrack(EMediaTrackTypes::Video);
+			auto VideoTracks = Player->GetVideoTracks();
 
-				if (VideoTrack.IsValid())
-				{
-					VideoTrackIndex = VideoTrack->GetIndex();
-				}
+			if (VideoTracks.IsValidIndex(VideoTrackIndex))
+			{
+				VideoTrack = VideoTracks[VideoTrackIndex];
 			}
-			else
+			else if (VideoTracks.Num() > 0)
 			{
-				VideoTrack = Player->GetTrack(VideoTrackIndex, EMediaTrackTypes::Video);
+				VideoTrack = VideoTracks[0];
+				VideoTrackIndex = 0;
 			}
 		}
 	}
 
 	if (VideoTrack.IsValid())
 	{
-		CachedDimensions = VideoTrack->GetVideoDetails().GetDimensions();
+		CachedDimensions = VideoTrack->GetDimensions();
 	}
 	else
 	{
@@ -232,7 +229,7 @@ void UMediaTexture::InitializeTrack()
 	// connect to new track
 	if (VideoTrack.IsValid())
 	{
-		VideoTrack->AddSink(VideoBuffer);
+		VideoTrack->GetStream().AddSink(VideoBuffer);
 	}
 }
 
diff --git a/Engine/Source/Runtime/MediaAssets/Public/MediaSoundWave.h b/Engine/Source/Runtime/MediaAssets/Public/MediaSoundWave.h
--- a/Engine/Source/Runtime/MediaAssets/Public/MediaSoundWave.h
+++ b/Engine/Source/Runtime/MediaAssets/Public/MediaSoundWave.h
@@ -8,8 +8,8 @@
 
 
 class FMediaSampleQueue;
+class IMediaAudioTrack;
 class IMediaPlayer;
-class IMediaTrack;
 class UMediaPlayer;
 
 
@@ -89,7 +89,7 @@ private:
 	TSharedRef<FMediaSampleQueue, ESPMode::ThreadSafe> AudioQueue;
 
 	/** Holds the selected audio track. */
-	TSharedPtr<IMediaTrack, ESPMode::ThreadSafe> AudioTrack;
+	TSharedPtr<IMediaAudioTrack, ESPMode::ThreadSafe> AudioTrack;
 
 	/** Holds the media player asset currently being used. */
 	UMediaPlayer* CurrentMediaPlayer;
diff --git a/Engine/Source/Runtime/MediaAssets/Public/MediaTexture.h b/Engine/Source/Runtime/MediaAssets/Public/MediaTexture.h
--- a/Engine/Source/Runtime/MediaAssets/Public/MediaTexture.h
+++ b/Engine/Source/Runtime/MediaAssets/Public/MediaTexture.h
@@ -8,7 +8,7 @@
 
 class FMediaSampleBuffer;
 class IMediaPlayer;
-class IMediaTrack;
+class IMediaVideoTrack;
 class UMediaPlayer;
 enum EPixelFormat;
 enum TextureAddress;
@@ -78,7 +78,7 @@ public:
 	 *
 	 * @return The selected video track, or nullptr if none is selected.
 	 */
-	TSharedPtr<IMediaTrack, ESPMode::ThreadSafe> GetVideoTrack() const
+	TSharedPtr<IMediaVideoTrack, ESPMode::ThreadSafe> GetVideoTrack() const
 	{
 		return VideoTrack;
 	}
@@ -147,5 +147,5 @@ private:
 	TSharedRef<FMediaSampleBuffer, ESPMode::ThreadSafe> VideoBuffer;
 
 	/** Holds the selected video track. */
-	TSharedPtr<IMediaTrack, ESPMode::ThreadSafe> VideoTrack;
+	TSharedPtr<IMediaVideoTrack, ESPMode::ThreadSafe> VideoTrack;
 };
diff --git a/Engine/Source/Runtime/Windows/D3D12RHI/Private/D3D12Util.cpp b/Engine/Source/Runtime/Windows/D3D12RHI/Private/D3D12Util.cpp
--- a/Engine/Source/Runtime/Windows/D3D12RHI/Private/D3D12Util.cpp
+++ b/Engine/Source/Runtime/Windows/D3D12RHI/Private/D3D12Util.cpp
@@ -294,6 +294,79 @@ FD3D12ResourceLocation* FD3D12DynamicBuffer::Unlock()
     return ResourceLocation;
 }
 
+
+void LogExecuteCommandLists(uint32 NumCommandLists, ID3D12CommandList *const *ppCommandLists)
+{
+	for (uint32 i = 0; i < NumCommandLists; i++)
+	{
+		ID3D12CommandList *const pCurrentCommandList = ppCommandLists[i];
+		UE_LOG(LogD3D12RHI, Log, TEXT("*** EXECUTE (CmdList: %016llX) %u/%u ***"), pCurrentCommandList, i + 1, NumCommandLists);
+	}
+}
+
+FString ConvertToResourceStateString(uint32 ResourceState)
+{
+	if (ResourceState == 0)
+	{
+		return FString(TEXT("D3D12_RESOURCE_STATE_COMMON"));
+	}
+
+	TCHAR *ResourceStateNames[] =
+	{
+		TEXT("D3D12_RESOURCE_STATE_VERTEX_AND_CONSTANT_BUFFER"	  ), 
+		TEXT("D3D12_RESOURCE_STATE_INDEX_BUFFER"				  ), 
+		TEXT("D3D12_RESOURCE_STATE_RENDER_TARGET"				  ), 
+		TEXT("D3D12_RESOURCE_STATE_UNORDERED_ACCESS"			  ), 
+		TEXT("D3D12_RESOURCE_STATE_DEPTH_WRITE"					  ), 
+		TEXT("D3D12_RESOURCE_STATE_DEPTH_READ"					  ), 
+		TEXT("D3D12_RESOURCE_STATE_NON_PIXEL_SHADER_RESOURCE"	  ), 
+		TEXT("D3D12_RESOURCE_STATE_PIXEL_SHADER_RESOURCE"		  ), 
+		TEXT("D3D12_RESOURCE_STATE_STREAM_OUT"					  ), 
+		TEXT("D3D12_RESOURCE_STATE_INDIRECT_ARGUMENT"			  ), 
+		TEXT("D3D12_RESOURCE_STATE_COPY_DEST"					  ), 
+		TEXT("D3D12_RESOURCE_STATE_COPY_SOURCE"					  ), 
+		TEXT("D3D12_RESOURCE_STATE_RESOLVE_DEST"				  ), 
+		TEXT("D3D12_RESOURCE_STATE_RESOLVE_SOURCE"				  ), 
+	};
+
+	FString ResourceStateString;
+	uint16 NumStates = 0;
+	for (uint16 i = 0; ResourceState && i < ARRAYSIZE(ResourceStateNames); i++)
+	{
+		if (ResourceState & 1)
+		{
+			if (NumStates > 0)
+			{
+				ResourceStateString += " | ";
+			}
+
+			ResourceStateString += ResourceStateNames[i];
+			NumStates++;
+		}
+		ResourceState = ResourceState >> 1;
+	}
+	return ResourceStateString;
+}
+
+void LogResourceBarriers(uint32 NumBarriers, D3D12_RESOURCE_BARRIER *pBarriers, ID3D12CommandList *const pCommandList)
+{
+	for (uint32 i = 0; i < NumBarriers; i++)
+	{
+		D3D12_RESOURCE_BARRIER &currentBarrier = pBarriers[i];
+		check(currentBarrier.Type == D3D12_RESOURCE_BARRIER_TYPE_TRANSITION);
+
+		FString StateBefore = ConvertToResourceStateString(static_cast<uint32>(currentBarrier.Transition.StateBefore));
+		FString StateAfter = ConvertToResourceStateString(static_cast<uint32>(currentBarrier.Transition.StateAfter));
+
+		UE_LOG(LogD3D12RHI, Log, TEXT("*** BARRIER (CmdList: %016llX) %u/%u: %016llX (Sub: %u), %s -> %s"), pCommandList, i + 1, NumBarriers,
+			currentBarrier.Transition.pResource,
+			currentBarrier.Transition.Subresource,
+			*StateBefore,
+			*StateAfter);
+	}
+}
+
+
 //==================================================================================================================================
 // CResourceState
 // Tracking of per-resource or per-subresource state
diff --git a/Engine/Source/Runtime/Windows/D3D12RHI/Public/D3D12Util.h b/Engine/Source/Runtime/Windows/D3D12RHI/Public/D3D12Util.h
--- a/Engine/Source/Runtime/Windows/D3D12RHI/Public/D3D12Util.h
+++ b/Engine/Source/Runtime/Windows/D3D12RHI/Public/D3D12Util.h
@@ -375,76 +375,10 @@ private: // Members
 	FCriticalSection CritSect;
 };
 
-void LogExecuteCommandLists(uint32 NumCommandLists, ID3D12CommandList *const *ppCommandLists)
-{
-	for (uint32 i = 0; i < NumCommandLists; i++)
-	{
-		ID3D12CommandList *const pCurrentCommandList = ppCommandLists[i];
-		UE_LOG(LogD3D12RHI, Log, TEXT("*** EXECUTE (CmdList: %016llX) %u/%u ***"), pCurrentCommandList, i + 1, NumCommandLists);
-	}
-}
+void LogExecuteCommandLists(uint32 NumCommandLists, ID3D12CommandList *const *ppCommandLists);
+FString ConvertToResourceStateString(uint32 ResourceState);
+void LogResourceBarriers(uint32 NumBarriers, D3D12_RESOURCE_BARRIER *pBarriers, ID3D12CommandList *const pCommandList);
 
-FString ConvertToResourceStateString(uint32 ResourceState)
-{
-	if (ResourceState == 0)
-	{
-		return FString(TEXT("D3D12_RESOURCE_STATE_COMMON"));
-	}
-	
-	TCHAR *ResourceStateNames[] =
-	{
-		TEXT("D3D12_RESOURCE_STATE_VERTEX_AND_CONSTANT_BUFFER"	  ), 
-		TEXT("D3D12_RESOURCE_STATE_INDEX_BUFFER"				  ), 
-		TEXT("D3D12_RESOURCE_STATE_RENDER_TARGET"				  ), 
-		TEXT("D3D12_RESOURCE_STATE_UNORDERED_ACCESS"			  ), 
-		TEXT("D3D12_RESOURCE_STATE_DEPTH_WRITE"					  ), 
-		TEXT("D3D12_RESOURCE_STATE_DEPTH_READ"					  ), 
-		TEXT("D3D12_RESOURCE_STATE_NON_PIXEL_SHADER_RESOURCE"	  ), 
-		TEXT("D3D12_RESOURCE_STATE_PIXEL_SHADER_RESOURCE"		  ), 
-		TEXT("D3D12_RESOURCE_STATE_STREAM_OUT"					  ), 
-		TEXT("D3D12_RESOURCE_STATE_INDIRECT_ARGUMENT"			  ), 
-		TEXT("D3D12_RESOURCE_STATE_COPY_DEST"					  ), 
-		TEXT("D3D12_RESOURCE_STATE_COPY_SOURCE"					  ), 
-		TEXT("D3D12_RESOURCE_STATE_RESOLVE_DEST"				  ), 
-		TEXT("D3D12_RESOURCE_STATE_RESOLVE_SOURCE"				  ), 
-	};
-
-	FString ResourceStateString;
-	uint16 NumStates = 0;
-	for (uint16 i = 0; ResourceState && i < ARRAYSIZE(ResourceStateNames); i++)
-	{
-		if (ResourceState & 1)
-		{
-			if (NumStates > 0)
-			{
-				ResourceStateString += " | ";
-			}
-
-			ResourceStateString += ResourceStateNames[i];
-			NumStates++;
-		}
-		ResourceState = ResourceState >> 1;
-	}
-	return ResourceStateString;
-}
-
-void LogResourceBarriers(uint32 NumBarriers, D3D12_RESOURCE_BARRIER *pBarriers, ID3D12CommandList *const pCommandList)
-{
-	for (uint32 i = 0; i < NumBarriers; i++)
-	{
-		D3D12_RESOURCE_BARRIER &currentBarrier = pBarriers[i];
-		check(currentBarrier.Type == D3D12_RESOURCE_BARRIER_TYPE_TRANSITION);
-
-		FString StateBefore = ConvertToResourceStateString(static_cast<uint32>(currentBarrier.Transition.StateBefore));
-		FString StateAfter = ConvertToResourceStateString(static_cast<uint32>(currentBarrier.Transition.StateAfter));
-
-		UE_LOG(LogD3D12RHI, Log, TEXT("*** BARRIER (CmdList: %016llX) %u/%u: %016llX (Sub: %u), %s -> %s"), pCommandList, i + 1, NumBarriers,
-			currentBarrier.Transition.pResource,
-			currentBarrier.Transition.Subresource,
-			*StateBefore,
-			*StateAfter);
-	}
-}
 
 // Custom resource states
 // To Be Determined (TBD) means we need to fill out a resource barrier before the command list is executed.
