Media: Moved asset detail customizations into MediaPlayerEditor module

diff --git a/Engine/Plugins/Media/MediaPlayerEditor/Source/MediaPlayerEditor/MediaPlayerEditor.Build.cs b/Engine/Plugins/Media/MediaPlayerEditor/Source/MediaPlayerEditor/MediaPlayerEditor.Build.cs
--- a/Engine/Plugins/Media/MediaPlayerEditor/Source/MediaPlayerEditor/MediaPlayerEditor.Build.cs
+++ b/Engine/Plugins/Media/MediaPlayerEditor/Source/MediaPlayerEditor/MediaPlayerEditor.Build.cs
@@ -19,6 +19,7 @@ public class MediaPlayerEditor : ModuleRules
 			new string[] {
 				"MediaPlayerEditor/Private",
 				"MediaPlayerEditor/Private/AssetTools",
+                "MediaPlayerEditor/Private/Customizations",
 				"MediaPlayerEditor/Private/Factories",
 				"MediaPlayerEditor/Private/Models",
 				"MediaPlayerEditor/Private/Styles",
@@ -31,6 +32,7 @@ public class MediaPlayerEditor : ModuleRules
 				"ContentBrowser",
 				"Core",
 				"CoreUObject",
+                "DesktopWidgets",
 				"EditorStyle",
 				"Engine",
 				"InputCore",
diff --git a/Engine/Plugins/Media/MediaPlayerEditor/Source/MediaPlayerEditor/Private/Customizations/MediaPlayerCustomization.cpp b/Engine/Plugins/Media/MediaPlayerEditor/Source/MediaPlayerEditor/Private/Customizations/MediaPlayerCustomization.cpp
new file mode 100644
--- /dev/null
+++ b/Engine/Plugins/Media/MediaPlayerEditor/Source/MediaPlayerEditor/Private/Customizations/MediaPlayerCustomization.cpp
@@ -0,0 +1,439 @@
+// Copyright 1998-2015 Epic Games, Inc. All Rights Reserved.
+
+#include "MediaPlayerEditorPrivatePCH.h"
+#include "IMediaModule.h"
+#include "IMediaPlayer.h"
+#include "IMediaPlayerFactory.h"
+#include "MediaPlayer.h"
+#include "MediaPlayerCustomization.h"
+#include "SFilePathPicker.h"
+
+
+#define LOCTEXT_NAMESPACE "FMediaPlayerCustomization"
+
+
+/* IDetailCustomization interface
+ *****************************************************************************/
+
+void FMediaPlayerCustomization::CustomizeDetails( IDetailLayoutBuilder& DetailBuilder )
+{
+	DetailBuilder.GetObjectsBeingCustomized(CustomizedMediaPlayers);
+
+	// customize 'Source' category
+	IDetailCategoryBuilder& SourceCategory = DetailBuilder.EditCategory("Source");
+	{
+		// URL
+		UrlProperty = DetailBuilder.GetProperty("URL");
+		{
+			IDetailPropertyRow& UrlRow = SourceCategory.AddProperty(UrlProperty);
+
+			UrlRow.DisplayName(LOCTEXT("URLLabel", "URL"));
+			UrlRow.CustomWidget()
+				.NameContent()
+				[
+					SNew(SHorizontalBox)
+
+					+ SHorizontalBox::Slot()
+						.AutoWidth()
+						.VAlign(VAlign_Center)
+						[
+							SNew(STextBlock)
+								.Font(IDetailLayoutBuilder::GetDetailFont())
+								.Text(LOCTEXT("FileOrUrlPropertyName", "File or URL"))
+								.ToolTipText(UrlProperty->GetToolTipText())
+						]
+
+					+ SHorizontalBox::Slot()
+						.FillWidth(1.0f)
+						.HAlign(HAlign_Left)
+						.VAlign(VAlign_Center)
+						.Padding(4.0f, 0.0f, 0.0f, 0.0f)
+						[
+							SNew(SImage)
+								.Image(FCoreStyle::Get().GetBrush("Icons.Warning"))
+								.ToolTipText(LOCTEXT("InvalidUrlPathWarning", "The current URL points to a file that does not exist or is not located inside the /Content/Movies/ directory."))
+								.Visibility(this, &FMediaPlayerCustomization::HandleUrlWarningIconVisibility)
+						]
+				]
+				.ValueContent()
+				.MaxDesiredWidth(0.0f)
+				.MinDesiredWidth(125.0f)
+				[
+					SNew(SFilePathPicker)
+						.BrowseButtonImage(FEditorStyle::GetBrush("PropertyWindow.Button_Ellipsis"))
+						.BrowseButtonStyle(FEditorStyle::Get(), "HoverHintOnly")
+						.BrowseButtonToolTip(LOCTEXT("UrlBrowseButtonToolTip", "Choose a file from this computer"))
+						.BrowseDirectory(FPaths::GameContentDir() / TEXT("Movies"))
+						.FilePath(this, &FMediaPlayerCustomization::HandleUrlPickerFilePath)
+						.FileTypeFilter(this, &FMediaPlayerCustomization::HandleUrlPickerFileTypeFilter)
+						.OnPathPicked(this, &FMediaPlayerCustomization::HandleUrlPickerPathPicked)
+						.ToolTipText(LOCTEXT("UrlToolTip", "The path to a media file on this computer, or a URL to a media source on the internet"))
+				];
+		}
+	}
+
+	// customize 'Information' category
+	IDetailCategoryBuilder& InformationCategory = DetailBuilder.EditCategory("Information", FText::GetEmpty(), ECategoryPriority::Uncommon);
+	{
+		// duration
+		InformationCategory.AddCustomRow(LOCTEXT("Duration", "Duration"))
+			.NameContent()
+			[
+				SNew(STextBlock)
+					.Text(LOCTEXT("Duration", "Duration"))
+					.Font(FEditorStyle::GetFontStyle(TEXT("PropertyWindow.NormalFont")))
+			]
+			.ValueContent()
+			.MaxDesiredWidth(0.0f)
+			[
+				SNew(STextBlock)
+					.Text(this, &FMediaPlayerCustomization::HandleDurationTextBlockText)
+					.Font(FEditorStyle::GetFontStyle(TEXT("PropertyWindow.NormalFont")))
+					.ToolTipText(LOCTEXT("DurationToolTip", "The total duration of this media, i.e. how long it plays"))
+			];
+
+		// forward rates (thinned)
+		InformationCategory.AddCustomRow(LOCTEXT("ForwardRatesThinned", "Forward Rates (Thinned)"))
+			.NameContent()
+			[
+				SNew(STextBlock)
+					.Text(LOCTEXT("ForwardRatesThinned", "Forward Rates (Thinned)"))
+					.Font(FEditorStyle::GetFontStyle(TEXT("PropertyWindow.NormalFont")))
+			]
+			.ValueContent()
+			.MaxDesiredWidth(0.0f)
+			[
+				SNew(STextBlock)
+					.Text(this, &FMediaPlayerCustomization::HandleSupportedRatesTextBlockText, EMediaPlaybackDirections::Forward, false)
+					.Font(FEditorStyle::GetFontStyle(TEXT("PropertyWindow.NormalFont")))
+					.ToolTipText(LOCTEXT("ThinnedForwardRateToolTip", "The rate at which this media can be played thinned (with skipping frames) in the forward direction."))
+			];
+
+		// forward rates (unthinned)
+		InformationCategory.AddCustomRow(LOCTEXT("ForwardRatesUnthinned", "Forward Rates (Unthinned)"))
+			.NameContent()
+			[
+				SNew(STextBlock)
+					.Text(LOCTEXT("ForwardRatesUnthinned", "Forward Rates (Unthinned)"))
+					.Font(FEditorStyle::GetFontStyle(TEXT("PropertyWindow.NormalFont")))
+			]
+			.ValueContent()
+			.MaxDesiredWidth(0.0f)
+			[
+				SNew(STextBlock)
+					.Text(this, &FMediaPlayerCustomization::HandleSupportedRatesTextBlockText, EMediaPlaybackDirections::Forward, true)
+					.Font(FEditorStyle::GetFontStyle(TEXT("PropertyWindow.NormalFont")))
+					.ToolTipText(LOCTEXT("UnthinnedForwardRateToolTip", "The rate at which this media can be played unthinned (without skipping frames) in the forward direction."))
+			];
+
+		// reverse rates (thinned)
+		InformationCategory.AddCustomRow(LOCTEXT("ReverseRatesThinned", "Reverse Rates (Thinned)"))
+			.NameContent()
+			[
+				SNew(STextBlock)
+					.Text(LOCTEXT("ReverseRatesThinned", "Reverse Rates (Thinned)"))
+					.Font(FEditorStyle::GetFontStyle(TEXT("PropertyWindow.NormalFont")))
+			]
+			.ValueContent()
+			.MaxDesiredWidth(0.0f)
+			[
+				SNew(STextBlock)
+					.Text(this, &FMediaPlayerCustomization::HandleSupportedRatesTextBlockText, EMediaPlaybackDirections::Reverse, false)
+					.Font(FEditorStyle::GetFontStyle(TEXT("PropertyWindow.NormalFont")))
+					.ToolTipText(LOCTEXT("ThinnedReverseRateToolTip", "The rate at which this media can be played thinned (with skipping frames) in the reverse direction."))
+			];
+
+		// reverse rates (unthinned)
+		InformationCategory.AddCustomRow(LOCTEXT("ReverseRatesUnthinned", "Reverse Rates (Unthinned)"))
+			.NameContent()
+			[
+				SNew(STextBlock)
+					.Text(LOCTEXT("ReverseRatesUnthinned", "Reverse Rates (Unthinned)"))
+					.Font(FEditorStyle::GetFontStyle(TEXT("PropertyWindow.NormalFont")))
+			]
+			.ValueContent()
+			.MaxDesiredWidth(0.0f)
+			[
+				SNew(STextBlock)
+					.Text(this, &FMediaPlayerCustomization::HandleSupportedRatesTextBlockText, EMediaPlaybackDirections::Reverse, true)
+					.Font(FEditorStyle::GetFontStyle(TEXT("PropertyWindow.NormalFont")))
+					.ToolTipText(LOCTEXT("UnthinnedReverseRateToolTip", "The rate at which this media can be played unthinned (without skipping frames) in the reverse direction."))
+			];
+
+		// supports scrubbing
+		InformationCategory.AddCustomRow(LOCTEXT("SupportsScrubbing", "Supports Scrubbing"))
+			.NameContent()
+			[
+				SNew(STextBlock)
+					.Text(LOCTEXT("SupportsScrubbing", "Supports Scrubbing"))
+					.Font(FEditorStyle::GetFontStyle(TEXT("PropertyWindow.NormalFont")))
+			]
+			.ValueContent()
+			.MaxDesiredWidth(0.0f)
+			[
+				SNew(STextBlock)
+					.Text(this, &FMediaPlayerCustomization::HandleSupportsScrubbingTextBlockText)
+					.Font(FEditorStyle::GetFontStyle(TEXT("PropertyWindow.NormalFont")))
+					.ToolTipText(LOCTEXT("ScrubbingToolTip", "Whether this media supports scrubbing, i.e. manually controlling the playback speed and position."))
+			];
+
+		// supports seeking
+		InformationCategory.AddCustomRow(LOCTEXT("SupportsSeeking", "Supports Seeking"))
+			.NameContent()
+			[
+				SNew(STextBlock)
+					.Text(LOCTEXT("SupportsSeeking", "Supports Seeking"))
+					.Font(FEditorStyle::GetFontStyle(TEXT("PropertyWindow.NormalFont")))
+			]
+			.ValueContent()
+			.MaxDesiredWidth(0.0f)
+			[
+				SNew(STextBlock)
+					.Text(this, &FMediaPlayerCustomization::HandleSupportsSeekingTextBlockText)
+					.Font(FEditorStyle::GetFontStyle(TEXT("PropertyWindow.NormalFont")))
+					.ToolTipText(LOCTEXT("SeekingToolTip", "Whether this media supports seeking, i.e. jumping to a specific position."))
+			];
+	}
+}
+
+
+/* FMediaPlayerCustomization callbacks
+ *****************************************************************************/
+
+FText FMediaPlayerCustomization::HandleDurationTextBlockText() const
+{
+	TOptional<FTimespan> Duration;
+
+	// ensure that all assets have the same value
+	for (auto MediaPlayerObject : CustomizedMediaPlayers)
+	{
+		if (MediaPlayerObject.IsValid())
+		{
+			FTimespan OtherDuration = Cast<UMediaPlayer>(MediaPlayerObject.Get())->GetDuration();
+
+			if (!Duration.IsSet())
+			{
+				Duration = OtherDuration;
+			}
+			else if (OtherDuration != Duration.GetValue())
+			{
+				return FText::GetEmpty();
+			}
+		}
+	}
+
+	if (!Duration.IsSet())
+	{
+		return FText::GetEmpty();
+	}
+
+	return FText::AsTimespan(Duration.GetValue());
+}
+
+
+FText FMediaPlayerCustomization::HandleSupportedRatesTextBlockText( EMediaPlaybackDirections Direction, bool Unthinned ) const
+{
+	TOptional<TRange<float>> Rates;
+
+	// ensure that all assets have the same value
+	for (auto MediaPlayerObject : CustomizedMediaPlayers)
+	{
+		if (MediaPlayerObject.IsValid())
+		{
+			IMediaPlayerPtr Player = Cast<UMediaPlayer>(MediaPlayerObject.Get())->GetPlayer();
+
+			if (!Player.IsValid())
+			{
+				return FText::GetEmpty();
+			}
+
+			TRange<float> OtherRates = Player->GetMediaInfo().GetSupportedRates(Direction, Unthinned);
+
+			if (!Rates.IsSet())
+			{
+				Rates = OtherRates;
+			}
+			else if (!Player.IsValid() || (OtherRates != Rates.GetValue()))
+			{
+				return FText::GetEmpty();
+			}
+		}
+	}
+
+	if (!Rates.IsSet())
+	{
+		return FText::GetEmpty();
+	}
+
+	if ((Rates.GetValue().GetLowerBoundValue() == 0.0f) && (Rates.GetValue().GetUpperBoundValue() == 0.0f))
+	{
+		return LOCTEXT("RateNotSupported", "Not supported");
+	}
+
+	return FText::Format(LOCTEXT("RatesFormat", "{0} to {1}"),
+		FText::AsNumber(Rates.GetValue().GetLowerBoundValue()),
+		FText::AsNumber(Rates.GetValue().GetUpperBoundValue()));
+}
+
+
+FText FMediaPlayerCustomization::HandleSupportsScrubbingTextBlockText() const
+{
+	TOptional<bool> SupportsScrubbing;
+
+	// ensure that all assets have the same value
+	for (auto MediaPlayerObject : CustomizedMediaPlayers)
+	{
+		if (MediaPlayerObject.IsValid())
+		{
+			bool OtherSupportsScrubbing = Cast<UMediaPlayer>(MediaPlayerObject.Get())->SupportsScrubbing();
+
+			if (!SupportsScrubbing.IsSet())
+			{
+				SupportsScrubbing = OtherSupportsScrubbing;
+			}
+			else if (OtherSupportsScrubbing != SupportsScrubbing.GetValue())
+			{
+				return FText::GetEmpty();
+			}
+		}
+	}
+
+	if (!SupportsScrubbing.IsSet())
+	{
+		return FText::GetEmpty();
+	}
+
+	return SupportsScrubbing.GetValue() ? GTrue : GFalse;
+}
+
+
+FText FMediaPlayerCustomization::HandleSupportsSeekingTextBlockText() const
+{
+	TOptional<bool> SupportsSeeking;
+
+	// ensure that all assets have the same value
+	for (auto MediaPlayerObject : CustomizedMediaPlayers)
+	{
+		if (MediaPlayerObject.IsValid())
+		{
+			bool OtherSupportsSeeking = Cast<UMediaPlayer>(MediaPlayerObject.Get())->SupportsSeeking();
+
+			if (!SupportsSeeking.IsSet())
+			{
+				SupportsSeeking = OtherSupportsSeeking;
+			}
+			else if (OtherSupportsSeeking != SupportsSeeking.GetValue())
+			{
+				return FText::GetEmpty();
+			}
+		}
+	}
+
+	if (!SupportsSeeking.IsSet())
+	{
+		return FText::GetEmpty();
+	}
+
+	return SupportsSeeking.GetValue() ? GTrue : GFalse;
+}
+
+
+FString FMediaPlayerCustomization::HandleUrlPickerFilePath() const
+{
+	FString Url;
+	UrlProperty->GetValue(Url);
+
+	return Url;
+}
+
+
+FString FMediaPlayerCustomization::HandleUrlPickerFileTypeFilter() const
+{
+	FString Filter = TEXT("All files (*.*)|*.*");
+
+	auto MediaModule = FModuleManager::GetModulePtr<IMediaModule>("Media");
+
+	if (MediaModule == nullptr)
+	{
+		return Filter;
+	}
+
+	FMediaFileTypes FileTypes;
+	MediaModule->GetSupportedFileTypes(FileTypes);
+
+	if (FileTypes.Num() == 0)
+	{
+		return Filter;
+	}
+
+	FString AllExtensions;
+	FString AllFilters;
+			
+	for (auto& Format : FileTypes)
+	{
+		if (!AllExtensions.IsEmpty())
+		{
+			AllExtensions += TEXT(";");
+		}
+
+		AllExtensions += TEXT("*.") + Format.Key;
+		AllFilters += TEXT("|") + Format.Value.ToString() + TEXT(" (*.") + Format.Key + TEXT(")|*.") + Format.Key;
+	}
+
+	Filter = TEXT("All movie files (") + AllExtensions + TEXT(")|") + AllExtensions + TEXT("|") + Filter + AllFilters;
+
+	return Filter;
+}
+
+
+void FMediaPlayerCustomization::HandleUrlPickerPathPicked( const FString& PickedPath )
+{
+	if (PickedPath.IsEmpty() || PickedPath.StartsWith(TEXT("./")) || PickedPath.Contains(TEXT("://")))
+	{
+		UrlProperty->SetValue(PickedPath);
+	}
+	else
+	{	
+		FString FullUrl = FPaths::ConvertRelativePathToFull(PickedPath);
+		const FString FullGameContentDir = FPaths::ConvertRelativePathToFull(FPaths::GameContentDir());
+
+		if (FullUrl.StartsWith(FullGameContentDir))
+		{
+			FPaths::MakePathRelativeTo(FullUrl, *FullGameContentDir);
+			FullUrl = FString(TEXT("./")) + FullUrl;
+		}
+
+		UrlProperty->SetValue(FullUrl);
+	}
+}
+
+
+EVisibility FMediaPlayerCustomization::HandleUrlWarningIconVisibility() const
+{
+	FString Url;
+
+	if ((UrlProperty->GetValue(Url) != FPropertyAccess::Success) || Url.IsEmpty() || Url.Contains(TEXT("://")))
+	{
+		return EVisibility::Hidden;
+	}
+
+	const FString FullMoviesPath = FPaths::ConvertRelativePathToFull(FPaths::GameContentDir() / TEXT("Movies"));
+	const FString FullUrl = FPaths::ConvertRelativePathToFull(FPaths::IsRelative(Url) ? FPaths::GameContentDir() / Url : Url);
+
+	if (FullUrl.StartsWith(FullMoviesPath))
+	{
+		if (FPaths::FileExists(FullUrl))
+		{
+			return EVisibility::Hidden;
+		}
+
+		// file doesn't exist
+		return EVisibility::Visible;
+	}
+
+	// file not inside Movies folder
+	return EVisibility::Visible;
+}
+
+
+#undef LOCTEXT_NAMESPACE
diff --git a/Engine/Plugins/Media/MediaPlayerEditor/Source/MediaPlayerEditor/Private/Customizations/MediaPlayerCustomization.h b/Engine/Plugins/Media/MediaPlayerEditor/Source/MediaPlayerEditor/Private/Customizations/MediaPlayerCustomization.h
new file mode 100644
--- /dev/null
+++ b/Engine/Plugins/Media/MediaPlayerEditor/Source/MediaPlayerEditor/Private/Customizations/MediaPlayerCustomization.h
@@ -0,0 +1,70 @@
+// Copyright 1998-2015 Epic Games, Inc. All Rights Reserved.
+
+#pragma once
+
+
+// forward declarations
+enum class EMediaPlaybackDirections;
+
+
+/**
+ * Implements a details view customization for the UMediaPlayer class.
+ */
+class FMediaPlayerCustomization
+	: public IDetailCustomization
+{
+public:
+
+	// IDetailCustomization interface
+
+	virtual void CustomizeDetails( IDetailLayoutBuilder& DetailBuilder ) override;
+
+public:
+
+	/**
+	 * Creates an instance of this class.
+	 *
+	 * @return The new instance.
+	 */
+	static TSharedRef<IDetailCustomization> MakeInstance()
+	{
+		return MakeShareable(new FMediaPlayerCustomization());
+	}
+
+private:
+
+	/** Callback for getting the text of the Duration text block. */
+	FText HandleDurationTextBlockText() const;
+
+	/** Callback for getting the text of the ForwardRates text block. */
+	FText HandleForwardRatesTextBlockText() const;
+
+	/** Callback for getting the text of a supported playback rate text block. */
+	FText HandleSupportedRatesTextBlockText( EMediaPlaybackDirections Direction, bool Unthinned ) const;
+
+	/** Callback for getting the text of the SupportsScrubbing text block. */
+	FText HandleSupportsScrubbingTextBlockText() const;
+
+	/** Callback for getting the text of the SupportsSeeking text block. */
+	FText HandleSupportsSeekingTextBlockText() const;
+
+	/** Callback for getting the selected path in the URL picker widget. */
+	FString HandleUrlPickerFilePath() const;
+
+	/** Callback for getting the file type filter for the URL picker. */
+	FString HandleUrlPickerFileTypeFilter() const;
+
+	/** Callback for picking a path in the URL picker. */
+	void HandleUrlPickerPathPicked( const FString& PickedPath );
+
+	/** Callback for getting the visibility of warning icon for invalid URLs. */
+	EVisibility HandleUrlWarningIconVisibility() const;
+
+private:
+
+	/** The collection of UMediaPlayer assets being customized. */
+	TArray<TWeakObjectPtr<UObject>> CustomizedMediaPlayers;
+
+	/** Pointer to the URL property handle. */
+	TSharedPtr<IPropertyHandle> UrlProperty;
+};
diff --git a/Engine/Plugins/Media/MediaPlayerEditor/Source/MediaPlayerEditor/Private/Customizations/MediaSoundWaveCustomization.cpp b/Engine/Plugins/Media/MediaPlayerEditor/Source/MediaPlayerEditor/Private/Customizations/MediaSoundWaveCustomization.cpp
new file mode 100644
--- /dev/null
+++ b/Engine/Plugins/Media/MediaPlayerEditor/Source/MediaPlayerEditor/Private/Customizations/MediaSoundWaveCustomization.cpp
@@ -0,0 +1,16 @@
+// Copyright 1998-2015 Epic Games, Inc. All Rights Reserved.
+
+#include "MediaPlayerEditorPrivatePCH.h"
+#include "IMediaPlayer.h"
+#include "IMediaStream.h"
+#include "MediaPlayer.h"
+#include "MediaSoundWaveCustomization.h"
+
+
+/* IDetailCustomization interface
+ *****************************************************************************/
+
+void FMediaSoundWaveCustomization::CustomizeDetails( IDetailLayoutBuilder& DetailBuilder )
+{
+
+}
diff --git a/Engine/Plugins/Media/MediaPlayerEditor/Source/MediaPlayerEditor/Private/Customizations/MediaSoundWaveCustomization.h b/Engine/Plugins/Media/MediaPlayerEditor/Source/MediaPlayerEditor/Private/Customizations/MediaSoundWaveCustomization.h
new file mode 100644
--- /dev/null
+++ b/Engine/Plugins/Media/MediaPlayerEditor/Source/MediaPlayerEditor/Private/Customizations/MediaSoundWaveCustomization.h
@@ -0,0 +1,29 @@
+// Copyright 1998-2015 Epic Games, Inc. All Rights Reserved.
+
+#pragma once
+
+
+/**
+ * Implements a details view customization for the UMediaSoundWave class.
+ */
+class FMediaSoundWaveCustomization
+	: public IDetailCustomization
+{
+public:
+
+	// IDetailCustomization interface
+
+	virtual void CustomizeDetails( IDetailLayoutBuilder& DetailBuilder ) override;
+
+public:
+
+	/**
+	 * Creates an instance of this class.
+	 *
+	 * @return The new instance.
+	 */
+	static TSharedRef<IDetailCustomization> MakeInstance()
+	{
+		return MakeShareable(new FMediaSoundWaveCustomization());
+	}
+};
diff --git a/Engine/Plugins/Media/MediaPlayerEditor/Source/MediaPlayerEditor/Private/Customizations/MediaTextureCustomization.cpp b/Engine/Plugins/Media/MediaPlayerEditor/Source/MediaPlayerEditor/Private/Customizations/MediaTextureCustomization.cpp
new file mode 100644
--- /dev/null
+++ b/Engine/Plugins/Media/MediaPlayerEditor/Source/MediaPlayerEditor/Private/Customizations/MediaTextureCustomization.cpp
@@ -0,0 +1,158 @@
+// Copyright 1998-2015 Epic Games, Inc. All Rights Reserved.
+
+#include "MediaPlayerEditorPrivatePCH.h"
+#include "IMediaPlayer.h"
+#include "IMediaStream.h"
+#include "IMediaAudioTrack.h"
+#include "IMediaCaptionTrack.h"
+#include "IMediaVideoTrack.h"
+#include "MediaPlayer.h"
+#include "MediaTextureCustomization.h"
+#include "PropertyHandle.h"
+
+
+#define LOCTEXT_NAMESPACE "FMediaTextureCustomization"
+
+
+/* IDetailCustomization interface
+ *****************************************************************************/
+
+void FMediaTextureCustomization::CustomizeDetails( IDetailLayoutBuilder& DetailBuilder )
+{
+	DetailBuilder.GetObjectsBeingCustomized(CustomizedMediaTextures);
+	MediaPlayerProperty = DetailBuilder.GetProperty("MediaPlayer");
+	VideoTrackIndexProperty = DetailBuilder.GetProperty("VideoTrackIndex");
+
+	// customize MediaPlayer category
+	IDetailCategoryBuilder& MediaPlayerCategory = DetailBuilder.EditCategory(TEXT("MediaPlayer"));
+	{
+		// video track index
+		IDetailPropertyRow& VideoTrackIndexRow = MediaPlayerCategory.AddProperty(VideoTrackIndexProperty);
+
+		VideoTrackIndexRow.DisplayName(LOCTEXT("VideoTrack", "Video Track"));
+		VideoTrackIndexRow.CustomWidget()
+			.NameContent()
+			[
+				VideoTrackIndexProperty->CreatePropertyNameWidget()
+			]
+			.ValueContent()
+			.MaxDesiredWidth(0.0f)
+			[
+				SNew(SComboButton)
+					.OnGetMenuContent(this, &FMediaTextureCustomization::HandleVideoTrackComboButtonMenuContent)
+ 					.ContentPadding(FMargin( 2.0f, 2.0f ))
+					.ButtonContent()
+					[
+						SNew(STextBlock) 
+							.Text(this, &FMediaTextureCustomization::HandleVideoTrackComboButtonText)
+							.Font(IDetailLayoutBuilder::GetDetailFont())
+					]
+			];
+	}
+}
+
+
+/* FMediaTextureCustomization callbacks
+ *****************************************************************************/
+
+TSharedRef<SWidget> FMediaTextureCustomization::HandleVideoTrackComboButtonMenuContent() const
+{
+	// get assigned media player asset
+	UObject* MediaPlayerObj = nullptr;
+	FPropertyAccess::Result Result = MediaPlayerProperty->GetValue(MediaPlayerObj);
+
+	if (Result == FPropertyAccess::MultipleValues)
+	{
+		return SNullWidget::NullWidget;
+	}
+
+	UMediaPlayer* MediaPlayer = Cast<UMediaPlayer>(MediaPlayerObj);
+
+	if (MediaPlayer == nullptr)
+	{
+		return SNullWidget::NullWidget;
+	}
+
+	// get media tracks
+	IMediaPlayerPtr Player = MediaPlayer->GetPlayer();
+
+	if (!Player.IsValid())
+	{
+		return SNullWidget::NullWidget;
+	}
+
+	const TArray<IMediaVideoTrackRef>& VideoTracks = Player->GetVideoTracks();
+
+	// populate the menu
+	FMenuBuilder MenuBuilder(true, nullptr);
+
+	for (int32 VideoTrackIndex = 0; VideoTrackIndex < VideoTracks.Num(); ++VideoTrackIndex)
+	{
+		FUIAction Action(FExecuteAction::CreateSP(this, &FMediaTextureCustomization::HandleVideoTrackComboButtonMenuEntryExecute, VideoTrackIndex));
+
+		MenuBuilder.AddMenuEntry(
+			VideoTracks[VideoTrackIndex]->GetStream().GetDisplayName(),
+			TAttribute<FText>(),
+			FSlateIcon(),
+			Action
+		);
+	}
+
+	return MenuBuilder.MakeWidget();
+}
+
+
+void FMediaTextureCustomization::HandleVideoTrackComboButtonMenuEntryExecute( int32 TrackIndex )
+{
+	VideoTrackIndexProperty->SetValue(TrackIndex);
+}
+
+
+FText FMediaTextureCustomization::HandleVideoTrackComboButtonText() const
+{
+	// get assigned media player asset
+	UObject* MediaPlayerObj = nullptr;
+	FPropertyAccess::Result Result = MediaPlayerProperty->GetValue(MediaPlayerObj);
+
+	if (Result != FPropertyAccess::Success)
+	{
+		return FText::GetEmpty();
+	}
+
+	UMediaPlayer* MediaPlayer = Cast<UMediaPlayer>(MediaPlayerObj);
+
+	if (MediaPlayer == nullptr)
+	{
+		return LOCTEXT("NoMediaPlayerSelected", "No Media Asset selected");
+	}
+
+	// get selected value
+	int32 VideoTrackIndex = -1;
+	Result = VideoTrackIndexProperty->GetValue(VideoTrackIndex);
+
+	if (Result != FPropertyAccess::Success)
+	{
+		return FText::GetEmpty();
+	}
+
+	// get selected media track
+	IMediaPlayerPtr Player = MediaPlayer->GetPlayer();
+
+	if (!Player.IsValid())
+	{
+		return LOCTEXT("NoMediaLoaded", "No media loaded");
+	}
+
+	auto VideoTracks = Player->GetVideoTracks();
+
+	if (!VideoTracks.IsValidIndex(VideoTrackIndex))
+	{
+		return LOCTEXT("SelectVideoTrack", "Select a Video Track...");
+	}
+
+	// get track name
+	return VideoTracks[VideoTrackIndex]->GetStream().GetDisplayName();
+}
+
+
+#undef LOCTEXT_NAMESPACE
diff --git a/Engine/Plugins/Media/MediaPlayerEditor/Source/MediaPlayerEditor/Private/Customizations/MediaTextureCustomization.h b/Engine/Plugins/Media/MediaPlayerEditor/Source/MediaPlayerEditor/Private/Customizations/MediaTextureCustomization.h
new file mode 100644
--- /dev/null
+++ b/Engine/Plugins/Media/MediaPlayerEditor/Source/MediaPlayerEditor/Private/Customizations/MediaTextureCustomization.h
@@ -0,0 +1,51 @@
+// Copyright 1998-2015 Epic Games, Inc. All Rights Reserved.
+
+#pragma once
+
+
+/**
+ * Implements a details view customization for the UMediaTexture class.
+ */
+class FMediaTextureCustomization
+	: public IDetailCustomization
+{
+public:
+
+	// IDetailCustomization interface
+
+	virtual void CustomizeDetails(IDetailLayoutBuilder& DetailBuilder) override;
+
+public:
+
+	/**
+	 * Creates an instance of this class.
+	 *
+	 * @return The new instance.
+	 */
+	static TSharedRef<IDetailCustomization> MakeInstance()
+	{
+		return MakeShareable(new FMediaTextureCustomization());
+	}
+
+private:
+
+	/** Callback for generating the menu content of the VideoTrack combo box. */
+	TSharedRef<SWidget> HandleVideoTrackComboButtonMenuContent() const;
+
+	/** Callback for selecting a track in the VideoTrack combo box. */
+	void HandleVideoTrackComboButtonMenuEntryExecute(int32 TrackIndex);
+
+	/** Callback for getting the text of the VideoTrack combo box. */
+	FText HandleVideoTrackComboButtonText() const;
+
+private:
+
+	/** The collection of media textures being customized */
+	TArray<TWeakObjectPtr<UObject>> CustomizedMediaTextures;
+
+	/** Pointer to the MediaPlayer property handle. */
+	TSharedPtr<IPropertyHandle> MediaPlayerProperty;
+
+	/** Pointer to the VideoTrackIndex property handle. */
+	TSharedPtr<IPropertyHandle> VideoTrackIndexProperty;
+};
diff --git a/Engine/Plugins/Media/MediaPlayerEditor/Source/MediaPlayerEditor/Private/MediaPlayerEditorModule.cpp b/Engine/Plugins/Media/MediaPlayerEditor/Source/MediaPlayerEditor/Private/MediaPlayerEditorModule.cpp
--- a/Engine/Plugins/Media/MediaPlayerEditor/Source/MediaPlayerEditor/Private/MediaPlayerEditorModule.cpp
+++ b/Engine/Plugins/Media/MediaPlayerEditor/Source/MediaPlayerEditor/Private/MediaPlayerEditorModule.cpp
@@ -45,6 +45,7 @@ public:
 		FMediaPlayerEditorCommands::Register();
 
 		RegisterAssetTools();
+		RegisterCustomizations();
 		RegisterMenuExtensions();
 		RegisterThumbnailRenderers();
 	}
@@ -52,6 +53,7 @@ public:
 	virtual void ShutdownModule() override
 	{
 		UnregisterAssetTools();
+		UnregisterCustomizations();
 		UnregisterMenuExtensions();
 		UnregisterThumbnailRenderers();
 	}
@@ -80,19 +82,6 @@ protected:
 		RegisteredAssetTypeActions.Add(Action);
 	}
 
-	/** Registers main menu and tool bar menu extensions. */
-	void RegisterMenuExtensions()
-	{
-		MenuExtensibilityManager = MakeShareable(new FExtensibilityManager);
-		ToolBarExtensibilityManager = MakeShareable(new FExtensibilityManager);
-	}
-
-	/** Registers asset thumbnail renderers .*/
-	void RegisterThumbnailRenderers()
-	{
-		UThumbnailManager::Get().RegisterCustomRenderer(UMediaTexture::StaticClass(), UTextureThumbnailRenderer::StaticClass());
-	}
-
 	/** Unregisters asset tool actions. */
 	void UnregisterAssetTools()
 	{
@@ -109,6 +98,35 @@ protected:
 		}
 	}
 
+protected:
+
+	/** Register details view customizations. */
+	void RegisterCustomizations()
+	{
+		FPropertyEditorModule& PropertyModule = FModuleManager::LoadModuleChecked<FPropertyEditorModule>("PropertyEditor");
+		PropertyModule.RegisterCustomClassLayout(UMediaPlayer::StaticClass()->GetFName(), FOnGetDetailCustomizationInstance::CreateStatic(&FMediaPlayerCustomization::MakeInstance));
+		PropertyModule.RegisterCustomClassLayout(UMediaSoundWave::StaticClass()->GetFName(), FOnGetDetailCustomizationInstance::CreateStatic(&FMediaSoundWaveCustomization::MakeInstance));
+		PropertyModule.RegisterCustomClassLayout(UMediaTexture::StaticClass()->GetFName(), FOnGetDetailCustomizationInstance::CreateStatic(&FMediaTextureCustomization::MakeInstance));
+	}
+
+	/** Unregister details view customizations. */
+	void UnregisterCustomizations()
+	{
+		FPropertyEditorModule& PropertyModule = FModuleManager::LoadModuleChecked<FPropertyEditorModule>("PropertyEditor");
+		PropertyModule.UnregisterCustomClassLayout(UMediaPlayer::StaticClass()->GetFName());
+		PropertyModule.UnregisterCustomClassLayout(UMediaSoundWave::StaticClass()->GetFName());
+		PropertyModule.UnregisterCustomClassLayout(UMediaTexture::StaticClass()->GetFName());
+	}
+
+protected:
+
+	/** Registers main menu and tool bar menu extensions. */
+	void RegisterMenuExtensions()
+	{
+		MenuExtensibilityManager = MakeShareable(new FExtensibilityManager);
+		ToolBarExtensibilityManager = MakeShareable(new FExtensibilityManager);
+	}
+
 	/** Unregisters main menu and tool bar menu extensions. */
 	void UnregisterMenuExtensions()
 	{
@@ -116,6 +134,14 @@ protected:
 		ToolBarExtensibilityManager.Reset();
 	}
 
+protected:
+
+	/** Registers asset thumbnail renderers .*/
+	void RegisterThumbnailRenderers()
+	{
+		UThumbnailManager::Get().RegisterCustomRenderer(UMediaTexture::StaticClass(), UTextureThumbnailRenderer::StaticClass());
+	}
+
 	/** Unregisters all asset thumbnail renderers. */
 	void UnregisterThumbnailRenderers()
 	{
diff --git a/Engine/Plugins/Media/MediaPlayerEditor/Source/MediaPlayerEditor/Private/MediaPlayerEditorPrivatePCH.h b/Engine/Plugins/Media/MediaPlayerEditor/Source/MediaPlayerEditor/Private/MediaPlayerEditorPrivatePCH.h
--- a/Engine/Plugins/Media/MediaPlayerEditor/Source/MediaPlayerEditor/Private/MediaPlayerEditorPrivatePCH.h
+++ b/Engine/Plugins/Media/MediaPlayerEditor/Source/MediaPlayerEditor/Private/MediaPlayerEditorPrivatePCH.h
@@ -12,6 +12,10 @@
 #include "Engine.h"
 #include "SlateStyle.h"
 #include "EditorStyle.h"
+#include "DetailCategoryBuilder.h"
+#include "DetailLayoutBuilder.h"
+#include "DetailWidgetRow.h"
+#include "IDetailPropertyRow.h"
 #include "IDetailsView.h"
 #include "IMediaStream.h"
 #include "IMediaAudioTrack.h"
@@ -35,6 +39,10 @@
 /* Private includes
  *****************************************************************************/
 
+#include "MediaPlayerCustomization.h"
+#include "MediaSoundWaveCustomization.h"
+#include "MediaTextureCustomization.h"
+
 #include "MediaPlayerFactory.h"
 #include "MediaPlayerFactoryNew.h"
 #include "MediaSoundWaveFactoryNew.h"
diff --git a/Engine/Source/Editor/DetailCustomizations/Private/DetailCustomizations.cpp b/Engine/Source/Editor/DetailCustomizations/Private/DetailCustomizations.cpp
--- a/Engine/Source/Editor/DetailCustomizations/Private/DetailCustomizations.cpp
+++ b/Engine/Source/Editor/DetailCustomizations/Private/DetailCustomizations.cpp
@@ -85,9 +85,6 @@
 #include "RawDistributionVectorStructCustomization.h"
 #include "CollisionProfileNameCustomization.h"
 #include "DocumentationActorDetails.h"
-#include "MediaPlayerCustomization.h"
-#include "MediaSoundWaveCustomization.h"
-#include "MediaTextureCustomization.h"
 #include "SoundWaveDetails.h"
 #include "AudioSettingsDetails.h"
 #include "DateTimeStructCustomization.h"
@@ -297,10 +294,6 @@ void FDetailCustomizationsModule::RegisterObjectCustomizations()
 	RegisterCustomClassLayout("SourceCodeAccessSettings", FOnGetDetailCustomizationInstance::CreateStatic(&FSourceCodeAccessSettingsDetails::MakeInstance));
 	RegisterCustomClassLayout("ParticleSystemComponent", FOnGetDetailCustomizationInstance::CreateStatic(&FParticleSystemComponentDetails::MakeInstance));
 
-	RegisterCustomClassLayout("MediaPlayer", FOnGetDetailCustomizationInstance::CreateStatic(&FMediaPlayerCustomization::MakeInstance));
-	RegisterCustomClassLayout("MediaSoundWave", FOnGetDetailCustomizationInstance::CreateStatic(&FMediaSoundWaveCustomization::MakeInstance));
-	RegisterCustomClassLayout("MediaTexture", FOnGetDetailCustomizationInstance::CreateStatic(&FMediaTextureCustomization::MakeInstance));
-
 	RegisterCustomClassLayout("FbxImportUI", FOnGetDetailCustomizationInstance::CreateStatic(&FFbxImportUIDetails::MakeInstance));
 
 	RegisterCustomClassLayout("ConfigHierarchyPropertyView", FOnGetDetailCustomizationInstance::CreateStatic(&FConfigPropertyHelperDetails::MakeInstance));
diff --git a/Engine/Source/Editor/DetailCustomizations/Private/MediaPlayerCustomization.cpp b/Engine/Source/Editor/DetailCustomizations/Private/MediaPlayerCustomization.cpp
deleted file mode 100644
diff --git a/Engine/Source/Editor/DetailCustomizations/Private/MediaPlayerCustomization.h b/Engine/Source/Editor/DetailCustomizations/Private/MediaPlayerCustomization.h
deleted file mode 100644
diff --git a/Engine/Source/Editor/DetailCustomizations/Private/MediaSoundWaveCustomization.cpp b/Engine/Source/Editor/DetailCustomizations/Private/MediaSoundWaveCustomization.cpp
deleted file mode 100644
diff --git a/Engine/Source/Editor/DetailCustomizations/Private/MediaSoundWaveCustomization.h b/Engine/Source/Editor/DetailCustomizations/Private/MediaSoundWaveCustomization.h
deleted file mode 100644
diff --git a/Engine/Source/Editor/DetailCustomizations/Private/MediaTextureCustomization.cpp b/Engine/Source/Editor/DetailCustomizations/Private/MediaTextureCustomization.cpp
deleted file mode 100644
diff --git a/Engine/Source/Editor/DetailCustomizations/Private/MediaTextureCustomization.h b/Engine/Source/Editor/DetailCustomizations/Private/MediaTextureCustomization.h
deleted file mode 100644
