Copying //UE4/Dev-Sequencer to //UE4/Dev-Main (Source: //UE4/Dev-Sequencer @ 3057646)

diff --git a/Engine/Plugins/Media/AvfMedia/Source/AvfMedia/Public/IAvfMediaModule.h b/Engine/Plugins/Media/AvfMedia/Source/AvfMedia/Public/IAvfMediaModule.h
new file mode 100644
--- /dev/null
+++ b/Engine/Plugins/Media/AvfMedia/Source/AvfMedia/Public/IAvfMediaModule.h
@@ -0,0 +1,30 @@
+// Copyright 1998-2016 Epic Games, Inc. All Rights Reserved.
+
+#pragma once
+
+#include "ModuleInterface.h"
+
+
+class IMediaPlayer;
+
+
+/**
+ * Interface for the AvfMedia module.
+ */
+class IAvfMediaModule
+	: public IModuleInterface
+{
+public:
+
+	/**
+	 * Creates a AV Foundation based media player.
+	 *
+	 * @return A new media player, or nullptr if a player couldn't be created.
+	 */
+	virtual TSharedPtr<IMediaPlayer> CreatePlayer() = 0;
+
+public:
+
+	/** Virtual destructor. */
+	virtual ~IAvfMediaModule() { }
+};
diff --git a/Engine/Plugins/Media/AvfMedia/Source/AvfMediaEditor/AvfMediaEditor.Build.cs b/Engine/Plugins/Media/AvfMedia/Source/AvfMediaEditor/AvfMediaEditor.Build.cs
new file mode 100644
--- /dev/null
+++ b/Engine/Plugins/Media/AvfMedia/Source/AvfMediaEditor/AvfMediaEditor.Build.cs
@@ -0,0 +1,25 @@
+// Copyright 1998-2016 Epic Games, Inc. All Rights Reserved.
+
+namespace UnrealBuildTool.Rules
+{
+	public class AvfMediaEditor : ModuleRules
+	{
+		public AvfMediaEditor(TargetInfo Target)
+		{
+			PrivateDependencyModuleNames.AddRange(
+				new string[] {
+					"Core",
+					"CoreUObject",
+					"MediaAssets",
+					"UnrealEd",
+                }
+			);
+
+			PrivateIncludePaths.AddRange(
+				new string[] {
+					"AvfMediaEditor/Private",
+				}
+			);
+		}
+	}
+}
diff --git a/Engine/Plugins/Media/AvfMedia/Source/AvfMediaEditor/Private/AvfMediaEditorModule.cpp b/Engine/Plugins/Media/AvfMedia/Source/AvfMediaEditor/Private/AvfMediaEditorModule.cpp
new file mode 100644
--- /dev/null
+++ b/Engine/Plugins/Media/AvfMedia/Source/AvfMediaEditor/Private/AvfMediaEditorModule.cpp
@@ -0,0 +1,22 @@
+// Copyright 1998-2016 Epic Games, Inc. All Rights Reserved.
+
+#include "AvfMediaEditorPCH.h"
+#include "ModuleInterface.h"
+
+
+/**
+ * Implements the AvfMediaEditor module.
+ */
+class FAvfMediaEditorModule
+	: public IModuleInterface
+{
+public:
+
+	//~ IModuleInterface interface
+
+	virtual void StartupModule() override { }
+	virtual void ShutdownModule() override { }
+};
+
+
+IMPLEMENT_MODULE(FAvfMediaEditorModule, AvfMediaEditor);
diff --git a/Engine/Plugins/Media/AvfMedia/Source/AvfMediaEditor/Private/AvfMediaEditorPCH.h b/Engine/Plugins/Media/AvfMedia/Source/AvfMediaEditor/Private/AvfMediaEditorPCH.h
new file mode 100644
--- /dev/null
+++ b/Engine/Plugins/Media/AvfMedia/Source/AvfMediaEditor/Private/AvfMediaEditorPCH.h
@@ -0,0 +1,6 @@
+// Copyright 1998-2016 Epic Games, Inc. All Rights Reserved.
+
+#pragma once
+
+#include "Runtime/Core/Public/Core.h"
+#include "Runtime/MediaAssets/Public/FileMediaSource.h"
diff --git a/Engine/Plugins/Media/AvfMedia/Source/AvfMediaFactory/AvfMediaFactory.Build.cs b/Engine/Plugins/Media/AvfMedia/Source/AvfMediaFactory/AvfMediaFactory.Build.cs
new file mode 100644
--- /dev/null
+++ b/Engine/Plugins/Media/AvfMedia/Source/AvfMediaFactory/AvfMediaFactory.Build.cs
@@ -0,0 +1,42 @@
+// Copyright 1998-2016 Epic Games, Inc. All Rights Reserved.
+
+namespace UnrealBuildTool.Rules
+{
+	public class AvfMediaFactory : ModuleRules
+	{
+		public AvfMediaFactory(TargetInfo Target)
+		{
+            DynamicallyLoadedModuleNames.AddRange(
+                new string[] {
+                    "Media",
+				}
+            );
+
+			PrivateDependencyModuleNames.AddRange(
+				new string[] {
+					"Core",
+                    "CoreUObject",
+                    "MediaAssets",
+				}
+			);
+
+			PrivateIncludePathModuleNames.AddRange(
+				new string[] {
+					"AvfMedia",
+                    "Media",
+				}
+			);
+
+			PrivateIncludePaths.AddRange(
+				new string[] {
+					"AvfMediaFactory/Private",
+				}
+			);
+
+			if ((Target.Platform == UnrealTargetPlatform.IOS) || (Target.Platform == UnrealTargetPlatform.Mac))
+			{
+				DynamicallyLoadedModuleNames.Add("AvfMedia");
+			}
+		}
+	}
+}
diff --git a/Engine/Plugins/Media/AvfMedia/Source/AvfMediaFactory/Private/AvfMediaFactoryPCH.h b/Engine/Plugins/Media/AvfMedia/Source/AvfMediaFactory/Private/AvfMediaFactoryPCH.h
new file mode 100644
--- /dev/null
+++ b/Engine/Plugins/Media/AvfMedia/Source/AvfMediaFactory/Private/AvfMediaFactoryPCH.h
@@ -0,0 +1,14 @@
+// Copyright 1998-2016 Epic Games, Inc. All Rights Reserved.
+
+#pragma once
+
+#include "Runtime/Core/Public/Core.h"
+#include "Runtime/Core/Public/Modules/ModuleManager.h"
+#include "Runtime/Media/Public/IMediaModule.h"
+#include "Runtime/Media/Public/IMediaOptions.h"
+
+#include "../../AvfMedia/Public/IAvfMediaModule.h"
+
+
+/** Log category for the AvfMediaFactory module. */
+DECLARE_LOG_CATEGORY_EXTERN(LogAvfMediaFactory, Log, All);
diff --git a/Engine/Plugins/Media/MediaPlayerEditor/MediaPlayerEditor.uplugin b/Engine/Plugins/Media/MediaPlayerEditor/MediaPlayerEditor.uplugin
--- a/Engine/Plugins/Media/MediaPlayerEditor/MediaPlayerEditor.uplugin
+++ b/Engine/Plugins/Media/MediaPlayerEditor/MediaPlayerEditor.uplugin
@@ -7,7 +7,8 @@
 	"Category" : "Editor",
 	"CreatedBy" : "Epic Games, Inc.",
 	"CreatedByURL" : "http://epicgames.com",
-	"DocsURL" : "https://forums.unrealengine.com/showthread.php?46879-Media-Framework-Documentation-for-4-5-Preview",
+  "DocsURL" : "https://forums.unrealengine.com/showthread.php?46879-Media-Framework-Documentation-for-4-5-Preview",
+  "EngineVersion" : "4.13.0",
 	"MarketplaceURL" : "",
 	"SupportURL" : "",
 	"EnabledByDefault" : true,
diff --git a/Engine/Plugins/Media/MediaPlayerEditor/Source/MediaPlayerEditor/MediaPlayerEditor.Build.cs b/Engine/Plugins/Media/MediaPlayerEditor/Source/MediaPlayerEditor/MediaPlayerEditor.Build.cs
--- a/Engine/Plugins/Media/MediaPlayerEditor/Source/MediaPlayerEditor/MediaPlayerEditor.Build.cs
+++ b/Engine/Plugins/Media/MediaPlayerEditor/Source/MediaPlayerEditor/MediaPlayerEditor.Build.cs
@@ -22,8 +22,9 @@ public class MediaPlayerEditor : ModuleRules
                 "MediaPlayerEditor/Private/Customizations",
 				"MediaPlayerEditor/Private/Factories",
 				"MediaPlayerEditor/Private/Models",
-				"MediaPlayerEditor/Private/Styles",
-				"MediaPlayerEditor/Private/Widgets",
+				"MediaPlayerEditor/Private/Shared",
+                "MediaPlayerEditor/Private/Toolkits",
+                "MediaPlayerEditor/Private/Widgets",
 			}
 		);
 
@@ -32,8 +33,10 @@ public class MediaPlayerEditor : ModuleRules
 				"ContentBrowser",
 				"Core",
 				"CoreUObject",
+                "DesktopPlatform",
                 "DesktopWidgets",
 				"EditorStyle",
+                "EditorWidgets",
 				"Engine",
 				"InputCore",
 				"MediaAssets",
diff --git a/Engine/Plugins/Media/MediaPlayerEditor/Source/MediaPlayerEditor/Private/AssetTools/FileMediaSourceActions.cpp b/Engine/Plugins/Media/MediaPlayerEditor/Source/MediaPlayerEditor/Private/AssetTools/FileMediaSourceActions.cpp
new file mode 100644
--- /dev/null
+++ b/Engine/Plugins/Media/MediaPlayerEditor/Source/MediaPlayerEditor/Private/AssetTools/FileMediaSourceActions.cpp
@@ -0,0 +1,59 @@
+// Copyright 1998-2016 Epic Games, Inc. All Rights Reserved.
+
+#include "MediaPlayerEditorPCH.h"
+#include "FileMediaSourceActions.h"
+
+
+#define LOCTEXT_NAMESPACE "AssetTypeActions"
+
+
+/* FFileMediaSourceActions constructors
+ *****************************************************************************/
+
+FFileMediaSourceActions::FFileMediaSourceActions(const TSharedRef<ISlateStyle>& InStyle)
+	: Style(InStyle)
+{ }
+
+
+/* FAssetTypeActions_Base interface
+ *****************************************************************************/
+
+FText FFileMediaSourceActions::GetName() const
+{
+	return LOCTEXT("AssetTypeActions_FileMediaSource", "File Media Source");
+}
+
+
+UClass* FFileMediaSourceActions::GetSupportedClass() const
+{
+	return UFileMediaSource::StaticClass();
+}
+
+
+TSharedPtr<class SWidget> FFileMediaSourceActions::GetThumbnailOverlay(const FAssetData& AssetData) const
+{
+	TWeakObjectPtr<UFileMediaSource> FileMediaSource = Cast<UFileMediaSource>(AssetData.GetAsset());
+
+	if (!FileMediaSource.IsValid())
+	{
+		return nullptr;
+	}
+
+	return SNew(SBorder)
+		.BorderImage(FEditorStyle::GetNoBrush())
+		.Visibility_Lambda([=]() -> EVisibility {
+			return (FileMediaSource.IsValid() && FileMediaSource->PrecacheFile)
+				? EVisibility::HitTestInvisible
+				: EVisibility::Hidden;
+		})
+		.Padding(FMargin(0.0f, 0.0f, 2.0f, 7.0f))
+		.HAlign(HAlign_Right)
+		.VAlign(VAlign_Bottom)
+		[
+			SNew(SImage)
+				.Image(Style->GetBrush("MediaPlayerEditor.FileMediaSourcePrecached"))
+		];
+}
+
+
+#undef LOCTEXT_NAMESPACE
diff --git a/Engine/Plugins/Media/MediaPlayerEditor/Source/MediaPlayerEditor/Private/AssetTools/FileMediaSourceActions.h b/Engine/Plugins/Media/MediaPlayerEditor/Source/MediaPlayerEditor/Private/AssetTools/FileMediaSourceActions.h
new file mode 100644
--- /dev/null
+++ b/Engine/Plugins/Media/MediaPlayerEditor/Source/MediaPlayerEditor/Private/AssetTools/FileMediaSourceActions.h
@@ -0,0 +1,35 @@
+// Copyright 1998-2016 Epic Games, Inc. All Rights Reserved.
+
+#pragma once
+
+#include "MediaSourceActions.h"
+
+
+/**
+ * Implements an action for UMediaSource assets.
+ */
+class FFileMediaSourceActions
+	: public FMediaSourceActions
+{
+public:
+
+	/**
+	 * Creates and initializes a new instance.
+	 *
+	 * @param InStyle The style set to use for asset editor toolkits.
+	 */
+	FFileMediaSourceActions(const TSharedRef<ISlateStyle>& InStyle);
+
+public:
+
+	//~ FAssetTypeActions_Base interface
+
+	virtual FText GetName() const override;
+	virtual UClass* GetSupportedClass() const override;
+	virtual TSharedPtr<class SWidget> GetThumbnailOverlay(const FAssetData& AssetData) const override;
+
+private:
+
+	/** Pointer to the style set to use for toolkits. */
+	TSharedRef<ISlateStyle> Style;
+};
diff --git a/Engine/Plugins/Media/MediaPlayerEditor/Source/MediaPlayerEditor/Private/AssetTools/MediaPlayerActions.cpp b/Engine/Plugins/Media/MediaPlayerEditor/Source/MediaPlayerEditor/Private/AssetTools/MediaPlayerActions.cpp
--- a/Engine/Plugins/Media/MediaPlayerEditor/Source/MediaPlayerEditor/Private/AssetTools/MediaPlayerActions.cpp
+++ b/Engine/Plugins/Media/MediaPlayerEditor/Source/MediaPlayerEditor/Private/AssetTools/MediaPlayerActions.cpp
@@ -1,7 +1,10 @@
 // Copyright 1998-2016 Epic Games, Inc. All Rights Reserved.
 
-#include "MediaPlayerEditorPrivatePCH.h"
-#include "ContentBrowserModule.h"
+#include "MediaPlayerEditorPCH.h"
+#include "MediaPlayerActions.h"
+#include "MediaPlayerEditorToolkit.h"
+#include "MediaSoundWaveFactoryNew.h"
+#include "MediaTextureFactoryNew.h"
 
 
 #define LOCTEXT_NAMESPACE "AssetTypeActions"
@@ -10,12 +13,12 @@
 /* FMediaPlayerActions constructors
  *****************************************************************************/
 
-FMediaPlayerActions::FMediaPlayerActions( const TSharedRef<ISlateStyle>& InStyle )
+FMediaPlayerActions::FMediaPlayerActions(const TSharedRef<ISlateStyle>& InStyle)
 	: Style(InStyle)
 { }
 
 
-/* FAssetTypeActions_Base overrides
+/* FAssetTypeActions_Base interface
  *****************************************************************************/
 
 bool FMediaPlayerActions::CanFilter()
@@ -24,38 +27,16 @@ bool FMediaPlayerActions::CanFilter()
 }
 
 
-void FMediaPlayerActions::GetActions( const TArray<UObject*>& InObjects, FMenuBuilder& MenuBuilder )
+void FMediaPlayerActions::GetActions(const TArray<UObject*>& InObjects, FMenuBuilder& MenuBuilder)
 {
 	FAssetTypeActions_Base::GetActions(InObjects, MenuBuilder);
 
 	auto MediaPlayers = GetTypedWeakObjectPtrs<UMediaPlayer>(InObjects);
 
-	MenuBuilder.AddMenuEntry(
-		LOCTEXT("MediaPlayer_PlayMovie", "Play Movie"),
-		LOCTEXT("MediaPlayer_PlayMovieToolTip", "Starts playback of the media."),
-		FSlateIcon( FEditorStyle::GetStyleSetName(), "MediaAsset.AssetActions.Play" ),
-		FUIAction(
-			FExecuteAction::CreateSP(this, &FMediaPlayerActions::HandlePlayMovieActionExecute, MediaPlayers),
-			FCanExecuteAction()
-		)
-	);
-
-	MenuBuilder.AddMenuEntry(
-		LOCTEXT("MediaPlayer_PauseMovie", "Pause Movie"),
-		LOCTEXT("MediaPlayer_PauseMovieToolTip", "Pauses playback of the media."),
-		FSlateIcon( FEditorStyle::GetStyleSetName(), "MediaAsset.AssetActions.Pause" ),
-		FUIAction(
-			FExecuteAction::CreateSP(this, &FMediaPlayerActions::HandlePauseMovieActionExecute, MediaPlayers),
-			FCanExecuteAction()
-		)
-	);
-
-	MenuBuilder.AddMenuSeparator();
-
 	MenuBuilder.AddMenuEntry(
 		LOCTEXT("MediaPlayer_CreateMediaSoundWave", "Create Media Sound Wave"),
-		LOCTEXT("MediaPlayer_CreateMediaSoundWaveTooltip", "Creates a new MediaSoundWave using this MediaPlayer asset."),
-		FSlateIcon( FEditorStyle::GetStyleSetName(), "ClassIcon.MediaSoundWave" ),
+		LOCTEXT("MediaPlayer_CreateMediaSoundWaveTooltip", "Creates a new MediaSoundWave and assigns it to this MediaPlayer asset."),
+		FSlateIcon(FEditorStyle::GetStyleSetName(), "ClassIcon.MediaSoundWave"),
 		FUIAction(
 			FExecuteAction::CreateSP(this, &FMediaPlayerActions::ExecuteCreateMediaSoundWave, MediaPlayers),
 			FCanExecuteAction()
@@ -64,21 +45,19 @@ void FMediaPlayerActions::GetActions( const TArray<UObject*>& InObjects, FMenuBu
 
 	MenuBuilder.AddMenuEntry(
 		LOCTEXT("MediaPlayer_CreateMediaTexture", "Create Media Texture"),
-		LOCTEXT("MediaPlayer_CreateMediaTextureTooltip", "Creates a new MediaTexture using this MediaPlayer asset."),
-		FSlateIcon( FEditorStyle::GetStyleSetName(), "ClassIcon.MediaTexture" ),
+		LOCTEXT("MediaPlayer_CreateMediaTextureTooltip", "Creates a new MediaTexture and assigns it to this MediaPlayer asset."),
+		FSlateIcon(FEditorStyle::GetStyleSetName(), "ClassIcon.MediaTexture"),
 		FUIAction(
 			FExecuteAction::CreateSP(this, &FMediaPlayerActions::ExecuteCreateMediaTexture, MediaPlayers),
 			FCanExecuteAction()
 		)
 	);
-
-	FAssetTypeActions_Base::GetActions(InObjects, MenuBuilder);
 }
 
 
 uint32 FMediaPlayerActions::GetCategories()
 {
-	return EAssetTypeCategories::Misc;
+	return EAssetTypeCategories::Media;
 }
 
 
@@ -100,13 +79,13 @@ FColor FMediaPlayerActions::GetTypeColor() const
 }
 
 
-bool FMediaPlayerActions::HasActions( const TArray<UObject*>& InObjects ) const
+bool FMediaPlayerActions::HasActions(const TArray<UObject*>& InObjects) const
 {
 	return true;
 }
 
 
-void FMediaPlayerActions::OpenAssetEditor( const TArray<UObject*>& InObjects, TSharedPtr<class IToolkitHost> EditWithinLevelEditor )
+void FMediaPlayerActions::OpenAssetEditor(const TArray<UObject*>& InObjects, TSharedPtr<IToolkitHost> EditWithinLevelEditor)
 {
 	EToolkitMode::Type Mode = EditWithinLevelEditor.IsValid()
 		? EToolkitMode::WorldCentric
@@ -128,144 +107,86 @@ void FMediaPlayerActions::OpenAssetEditor( const TArray<UObject*>& InObjects, TS
 /* FAssetTypeActions_MediaPlayer callbacks
  *****************************************************************************/
 
-void FMediaPlayerActions::ExecuteCreateMediaSoundWave( TArray<TWeakObjectPtr<UMediaPlayer>> Objects )
+void FMediaPlayerActions::ExecuteCreateMediaSoundWave(TArray<TWeakObjectPtr<UMediaPlayer>> Objects)
 {
+	const FString DefaultSuffix = TEXT("_Sound");
+
 	IContentBrowserSingleton& ContentBrowserSingleton = FModuleManager::LoadModuleChecked<FContentBrowserModule>("ContentBrowser").Get();
-	const FString DefaultSuffix = TEXT("_Wav");
+	TArray<UObject*> ObjectsToSync;
 
-	if (Objects.Num() == 1)
+	for (auto ObjIt = Objects.CreateConstIterator(); ObjIt; ++ObjIt)
 	{
-		auto Object = Objects[0].Get();
+		auto MediaPlayer = (*ObjIt).Get();
 
-		if (Object != nullptr)
+		if (MediaPlayer == nullptr)
 		{
-			// Determine an appropriate name
-			FString Name;
-			FString PackagePath;
-			CreateUniqueAssetName(Object->GetOutermost()->GetName(), DefaultSuffix, PackagePath, Name);
+			continue;
+		}
 
-			// Create the factory used to generate the asset
-			UMediaSoundWaveFactoryNew* Factory = NewObject<UMediaSoundWaveFactoryNew>();
-			Factory->InitialMediaPlayer = Object;
+		// generate unique name
+		FString Name;
+		FString PackageName;
 
-			ContentBrowserSingleton.CreateNewAsset(Name, FPackageName::GetLongPackagePath(PackagePath), UMediaSoundWave::StaticClass(), Factory);
-		}
-	}
-	else
-	{
-		TArray<UObject*> ObjectsToSync;
+		CreateUniqueAssetName(MediaPlayer->GetOutermost()->GetName(), DefaultSuffix, PackageName, Name);
 
-		for (auto ObjIt = Objects.CreateConstIterator(); ObjIt; ++ObjIt)
-		{
-			auto Object = (*ObjIt).Get();
-
-			if (Object != nullptr)
-			{
-				FString Name;
-				FString PackageName;
-				CreateUniqueAssetName(Object->GetOutermost()->GetName(), DefaultSuffix, PackageName, Name);
-
-				// Create the factory used to generate the asset
-				UMediaSoundWaveFactoryNew* Factory = NewObject<UMediaSoundWaveFactoryNew>();
-				Factory->InitialMediaPlayer = Object;
-
-				FAssetToolsModule& AssetToolsModule = FModuleManager::GetModuleChecked<FAssetToolsModule>("AssetTools");
-				UObject* NewAsset = AssetToolsModule.Get().CreateAsset(Name, FPackageName::GetLongPackagePath(PackageName), UMediaSoundWave::StaticClass(), Factory);
-
-				if (NewAsset != nullptr)
-				{
-					ObjectsToSync.Add(NewAsset);
-				}
-			}
-		}
+		// create asset
+		auto Factory = NewObject<UMediaSoundWaveFactoryNew>();
+
+		FAssetToolsModule& AssetToolsModule = FModuleManager::GetModuleChecked<FAssetToolsModule>("AssetTools");
+		UObject* NewAsset = AssetToolsModule.Get().CreateAsset(Name, FPackageName::GetLongPackagePath(PackageName), UMediaSoundWave::StaticClass(), Factory);
 
-		if (ObjectsToSync.Num() > 0)
+		if (NewAsset != nullptr)
 		{
-			ContentBrowserSingleton.SyncBrowserToAssets(ObjectsToSync);
+			MediaPlayer->SetSoundWave(Cast<UMediaSoundWave>(NewAsset));
+			ObjectsToSync.Add(NewAsset);
 		}
 	}
+
+	if (ObjectsToSync.Num() > 0)
+	{
+		ContentBrowserSingleton.SyncBrowserToAssets(ObjectsToSync);
+	}
 }
 
 
-void FMediaPlayerActions::ExecuteCreateMediaTexture( TArray<TWeakObjectPtr<UMediaPlayer>> Objects )
+void FMediaPlayerActions::ExecuteCreateMediaTexture(TArray<TWeakObjectPtr<UMediaPlayer>> Objects)
 {
+	const FString DefaultSuffix = TEXT("_Video");
+
 	IContentBrowserSingleton& ContentBrowserSingleton = FModuleManager::LoadModuleChecked<FContentBrowserModule>("ContentBrowser").Get();
-	const FString DefaultSuffix = TEXT("_Tex");
+	TArray<UObject*> ObjectsToSync;
 
-	if (Objects.Num() == 1)
+	for (auto ObjIt = Objects.CreateConstIterator(); ObjIt; ++ObjIt)
 	{
-		auto Object = Objects[0].Get();
+		auto MediaPlayer = (*ObjIt).Get();
 
-		if (Object != nullptr)
+		if (MediaPlayer == nullptr)
 		{
-			// Determine an appropriate name
-			FString Name;
-			FString PackagePath;
-			CreateUniqueAssetName(Object->GetOutermost()->GetName(), DefaultSuffix, PackagePath, Name);
-
-			// Create the factory used to generate the asset
-			UMediaTextureFactoryNew* Factory = NewObject<UMediaTextureFactoryNew>();
-			Factory->InitialMediaPlayer = Object;
-
-			ContentBrowserSingleton.CreateNewAsset(Name, FPackageName::GetLongPackagePath(PackagePath), UMediaTexture::StaticClass(), Factory);
+			continue;
 		}
-	}
-	else
-	{
-		TArray<UObject*> ObjectsToSync;
 
-		for (auto ObjIt = Objects.CreateConstIterator(); ObjIt; ++ObjIt)
-		{
-			auto Object = (*ObjIt).Get();
-
-			if (Object != nullptr)
-			{
-				FString Name;
-				FString PackageName;
-				CreateUniqueAssetName(Object->GetOutermost()->GetName(), DefaultSuffix, PackageName, Name);
-
-				// Create the factory used to generate the asset
-				UMediaTextureFactoryNew* Factory = NewObject<UMediaTextureFactoryNew>();
-				Factory->InitialMediaPlayer = Object;
-
-				FAssetToolsModule& AssetToolsModule = FModuleManager::GetModuleChecked<FAssetToolsModule>("AssetTools");
-				UObject* NewAsset = AssetToolsModule.Get().CreateAsset(Name, FPackageName::GetLongPackagePath(PackageName), UMediaTexture::StaticClass(), Factory);
-
-				if (NewAsset != nullptr)
-				{
-					ObjectsToSync.Add(NewAsset);
-				}
-			}
-		}
+		// generate unique name
+		FString Name;
+		FString PackageName;
 
-		if (ObjectsToSync.Num() > 0)
-		{
-			ContentBrowserSingleton.SyncBrowserToAssets(ObjectsToSync);
-		}
-	}
-}
+		CreateUniqueAssetName(MediaPlayer->GetOutermost()->GetName(), DefaultSuffix, PackageName, Name);
 
+		// create asset
+		auto Factory = NewObject<UMediaTextureFactoryNew>();
 
-void FMediaPlayerActions::HandlePauseMovieActionExecute( TArray<TWeakObjectPtr<UMediaPlayer>> Objects )
-{
-	for (auto MediaPlayer : Objects)
-	{
-		if (MediaPlayer.IsValid())
+		FAssetToolsModule& AssetToolsModule = FModuleManager::GetModuleChecked<FAssetToolsModule>("AssetTools");
+		UObject* NewAsset = AssetToolsModule.Get().CreateAsset(Name, FPackageName::GetLongPackagePath(PackageName), UMediaTexture::StaticClass(), Factory);
+
+		if (NewAsset != nullptr)
 		{
-			MediaPlayer->Pause();
+			MediaPlayer->SetVideoTexture(Cast<UMediaTexture>(NewAsset));
+			ObjectsToSync.Add(NewAsset);
 		}
 	}
-}
-
 
-void FMediaPlayerActions::HandlePlayMovieActionExecute( TArray<TWeakObjectPtr<UMediaPlayer>> Objects )
-{
-	for (auto MediaPlayer : Objects)
+	if (ObjectsToSync.Num() > 0)
 	{
-		if (MediaPlayer.IsValid())
-		{
-			MediaPlayer->Play();
-		}
+		ContentBrowserSingleton.SyncBrowserToAssets(ObjectsToSync);
 	}
 }
 
diff --git a/Engine/Plugins/Media/MediaPlayerEditor/Source/MediaPlayerEditor/Private/AssetTools/MediaPlayerActions.h b/Engine/Plugins/Media/MediaPlayerEditor/Source/MediaPlayerEditor/Private/AssetTools/MediaPlayerActions.h
--- a/Engine/Plugins/Media/MediaPlayerEditor/Source/MediaPlayerEditor/Private/AssetTools/MediaPlayerActions.h
+++ b/Engine/Plugins/Media/MediaPlayerEditor/Source/MediaPlayerEditor/Private/AssetTools/MediaPlayerActions.h
@@ -2,7 +2,6 @@
 
 #pragma once
 
-
 #include "AssetTypeActions_Base.h"
 
 
@@ -19,34 +18,28 @@ public:
 	 *
 	 * @param InStyle The style set to use for asset editor toolkits.
 	 */
-	FMediaPlayerActions( const TSharedRef<ISlateStyle>& InStyle );
+	FMediaPlayerActions(const TSharedRef<ISlateStyle>& InStyle);
 
 public:
 
-	// FAssetTypeActions_Base overrides
+	//~ FAssetTypeActions_Base interface
 
 	virtual bool CanFilter() override;
-	virtual void GetActions( const TArray<UObject*>& InObjects, FMenuBuilder& MenuBuilder ) override;
+	virtual void GetActions(const TArray<UObject*>& InObjects, FMenuBuilder& MenuBuilder) override;
 	virtual uint32 GetCategories() override;
 	virtual FText GetName() const override;
 	virtual UClass* GetSupportedClass() const override;
 	virtual FColor GetTypeColor() const override;
-	virtual bool HasActions( const TArray<UObject*>& InObjects ) const override;
-	virtual void OpenAssetEditor( const TArray<UObject*>& InObjects, TSharedPtr<class IToolkitHost> EditWithinLevelEditor = TSharedPtr<IToolkitHost>() ) override;
+	virtual bool HasActions(const TArray<UObject*>& InObjects) const override;
+	virtual void OpenAssetEditor(const TArray<UObject*>& InObjects, TSharedPtr<IToolkitHost> EditWithinLevelEditor = TSharedPtr<IToolkitHost>()) override;
 
 private:
 
 	/** Callback for selecting CreateMediaSoundWave. */
-	void ExecuteCreateMediaSoundWave( TArray<TWeakObjectPtr<UMediaPlayer>> Objects );
+	void ExecuteCreateMediaSoundWave(TArray<TWeakObjectPtr<UMediaPlayer>> Objects);
 
 	/** Callback for selecting CreateMediaTexture. */
-	void ExecuteCreateMediaTexture( TArray<TWeakObjectPtr<UMediaPlayer>> Objects );	
-
-	/** Callback for executing the 'Play Movie' action. */
-	void HandlePauseMovieActionExecute( TArray<TWeakObjectPtr<UMediaPlayer>> Objects );
-
-	/** Callback for executing the 'Play Movie' action. */
-	void HandlePlayMovieActionExecute( TArray<TWeakObjectPtr<UMediaPlayer>> Objects );
+	void ExecuteCreateMediaTexture(TArray<TWeakObjectPtr<UMediaPlayer>> Objects);	
 
 private:
 
diff --git a/Engine/Plugins/Media/MediaPlayerEditor/Source/MediaPlayerEditor/Private/AssetTools/MediaPlaylistActions.cpp b/Engine/Plugins/Media/MediaPlayerEditor/Source/MediaPlayerEditor/Private/AssetTools/MediaPlaylistActions.cpp
new file mode 100644
--- /dev/null
+++ b/Engine/Plugins/Media/MediaPlayerEditor/Source/MediaPlayerEditor/Private/AssetTools/MediaPlaylistActions.cpp
@@ -0,0 +1,71 @@
+// Copyright 1998-2016 Epic Games, Inc. All Rights Reserved.
+
+#include "MediaPlayerEditorPCH.h"
+#include "MediaPlaylistActions.h"
+#include "MediaPlaylistEditorToolkit.h"
+
+
+#define LOCTEXT_NAMESPACE "AssetTypeActions"
+
+
+/* FMediaPlaylistActions constructors
+ *****************************************************************************/
+
+FMediaPlaylistActions::FMediaPlaylistActions(const TSharedRef<ISlateStyle>& InStyle)
+	: Style(InStyle)
+{ }
+
+
+/* FAssetTypeActions_Base interface
+ *****************************************************************************/
+
+bool FMediaPlaylistActions::CanFilter()
+{
+	return true;
+}
+
+
+uint32 FMediaPlaylistActions::GetCategories()
+{
+	return EAssetTypeCategories::Media;
+}
+
+
+FText FMediaPlaylistActions::GetName() const
+{
+	return NSLOCTEXT("AssetTypeActions", "AssetTypeActions_MediaPlaylist", "Media Playlist");
+}
+
+
+UClass* FMediaPlaylistActions::GetSupportedClass() const
+{
+	return UMediaPlaylist::StaticClass();
+}
+
+
+FColor FMediaPlaylistActions::GetTypeColor() const
+{
+	return FColor::Yellow;
+}
+
+
+void FMediaPlaylistActions::OpenAssetEditor(const TArray<UObject*>& InObjects, TSharedPtr<class IToolkitHost> EditWithinLevelEditor)
+{
+	EToolkitMode::Type Mode = EditWithinLevelEditor.IsValid()
+		? EToolkitMode::WorldCentric
+		: EToolkitMode::Standalone;
+
+	for (auto ObjIt = InObjects.CreateConstIterator(); ObjIt; ++ObjIt)
+	{
+		auto MediaPlaylist = Cast<UMediaPlaylist>(*ObjIt);
+
+		if (MediaPlaylist != nullptr)
+		{
+			TSharedRef<FMediaPlaylistEditorToolkit> EditorToolkit = MakeShareable(new FMediaPlaylistEditorToolkit(Style));
+			EditorToolkit->Initialize(MediaPlaylist, Mode, EditWithinLevelEditor);
+		}
+	}
+}
+
+
+#undef LOCTEXT_NAMESPACE
diff --git a/Engine/Plugins/Media/MediaPlayerEditor/Source/MediaPlayerEditor/Private/AssetTools/MediaPlaylistActions.h b/Engine/Plugins/Media/MediaPlayerEditor/Source/MediaPlayerEditor/Private/AssetTools/MediaPlaylistActions.h
new file mode 100644
--- /dev/null
+++ b/Engine/Plugins/Media/MediaPlayerEditor/Source/MediaPlayerEditor/Private/AssetTools/MediaPlaylistActions.h
@@ -0,0 +1,38 @@
+// Copyright 1998-2016 Epic Games, Inc. All Rights Reserved.
+
+#pragma once
+
+#include "AssetTypeActions_Base.h"
+
+
+/**
+ * Implements an action for UMediaPlaylist assets.
+ */
+class FMediaPlaylistActions
+	: public FAssetTypeActions_Base
+{
+public:
+
+	/**
+	 * Creates and initializes a new instance.
+	 *
+	 * @param InStyle The style set to use for asset editor toolkits.
+	 */
+	FMediaPlaylistActions(const TSharedRef<ISlateStyle>& InStyle);
+
+public:
+
+	//~ FAssetTypeActions_Base interface
+
+	virtual bool CanFilter() override;
+	virtual uint32 GetCategories() override;
+	virtual FText GetName() const override;
+	virtual UClass* GetSupportedClass() const override;
+	virtual FColor GetTypeColor() const override;
+	virtual void OpenAssetEditor(const TArray<UObject*>& InObjects, TSharedPtr<class IToolkitHost> EditWithinLevelEditor = TSharedPtr<IToolkitHost>()) override;
+
+private:
+
+	/** Pointer to the style set to use for toolkits. */
+	TSharedRef<ISlateStyle> Style;
+};
diff --git a/Engine/Plugins/Media/MediaPlayerEditor/Source/MediaPlayerEditor/Private/AssetTools/MediaSoundWaveActions.cpp b/Engine/Plugins/Media/MediaPlayerEditor/Source/MediaPlayerEditor/Private/AssetTools/MediaSoundWaveActions.cpp
--- a/Engine/Plugins/Media/MediaPlayerEditor/Source/MediaPlayerEditor/Private/AssetTools/MediaSoundWaveActions.cpp
+++ b/Engine/Plugins/Media/MediaPlayerEditor/Source/MediaPlayerEditor/Private/AssetTools/MediaSoundWaveActions.cpp
@@ -1,6 +1,7 @@
 // Copyright 1998-2016 Epic Games, Inc. All Rights Reserved.
 
-#include "MediaPlayerEditorPrivatePCH.h"
+#include "MediaPlayerEditorPCH.h"
+#include "MediaSoundWaveActions.h"
 
 
 #define LOCTEXT_NAMESPACE "AssetTypeActions"
diff --git a/Engine/Plugins/Media/MediaPlayerEditor/Source/MediaPlayerEditor/Private/AssetTools/MediaSourceActions.h b/Engine/Plugins/Media/MediaPlayerEditor/Source/MediaPlayerEditor/Private/AssetTools/MediaSourceActions.h
new file mode 100644
--- /dev/null
+++ b/Engine/Plugins/Media/MediaPlayerEditor/Source/MediaPlayerEditor/Private/AssetTools/MediaSourceActions.h
@@ -0,0 +1,24 @@
+// Copyright 1998-2016 Epic Games, Inc. All Rights Reserved.
+
+#pragma once
+
+#include "AssetTypeActions_Base.h"
+
+
+/**
+ * Implements an action for UMediaSource assets.
+ */
+class FMediaSourceActions
+	: public FAssetTypeActions_Base
+{
+public:
+
+	//~ FAssetTypeActions_Base interface
+
+	virtual FText GetAssetDescription(const class FAssetData& AssetData) const override;
+	virtual uint32 GetCategories() override;
+	virtual FText GetName() const override;
+	virtual UClass* GetSupportedClass() const override;
+	virtual FColor GetTypeColor() const override;
+	virtual TSharedPtr<class SWidget> GetThumbnailOverlay(const FAssetData& AssetData) const override;
+};
diff --git a/Engine/Plugins/Media/MediaPlayerEditor/Source/MediaPlayerEditor/Private/AssetTools/MediaTextureActions.cpp b/Engine/Plugins/Media/MediaPlayerEditor/Source/MediaPlayerEditor/Private/AssetTools/MediaTextureActions.cpp
--- a/Engine/Plugins/Media/MediaPlayerEditor/Source/MediaPlayerEditor/Private/AssetTools/MediaTextureActions.cpp
+++ b/Engine/Plugins/Media/MediaPlayerEditor/Source/MediaPlayerEditor/Private/AssetTools/MediaTextureActions.cpp
@@ -1,8 +1,7 @@
 // Copyright 1998-2016 Epic Games, Inc. All Rights Reserved.
 
-#include "MediaPlayerEditorPrivatePCH.h"
-#include "Editor/ContentBrowser/Public/ContentBrowserModule.h"
-#include "TextureEditor.h"
+#include "MediaPlayerEditorPCH.h"
+#include "MediaTextureActions.h"
 
 
 #define LOCTEXT_NAMESPACE "AssetTypeActions"
diff --git a/Engine/Plugins/Media/MediaPlayerEditor/Source/MediaPlayerEditor/Private/Customizations/MediaSoundWaveCustomization.cpp b/Engine/Plugins/Media/MediaPlayerEditor/Source/MediaPlayerEditor/Private/Customizations/MediaSoundWaveCustomization.cpp
--- a/Engine/Plugins/Media/MediaPlayerEditor/Source/MediaPlayerEditor/Private/Customizations/MediaSoundWaveCustomization.cpp
+++ b/Engine/Plugins/Media/MediaPlayerEditor/Source/MediaPlayerEditor/Private/Customizations/MediaSoundWaveCustomization.cpp
@@ -1,16 +1,13 @@
 // Copyright 1998-2016 Epic Games, Inc. All Rights Reserved.
 
-#include "MediaPlayerEditorPrivatePCH.h"
-#include "IMediaPlayer.h"
-#include "IMediaStream.h"
-#include "MediaPlayer.h"
+#include "MediaPlayerEditorPCH.h"
 #include "MediaSoundWaveCustomization.h"
 
 
 /* IDetailCustomization interface
  *****************************************************************************/
 
-void FMediaSoundWaveCustomization::CustomizeDetails( IDetailLayoutBuilder& DetailBuilder )
+void FMediaSoundWaveCustomization::CustomizeDetails(IDetailLayoutBuilder& DetailBuilder)
 {
 
 }
diff --git a/Engine/Plugins/Media/MediaPlayerEditor/Source/MediaPlayerEditor/Private/Customizations/MediaSoundWaveCustomization.h b/Engine/Plugins/Media/MediaPlayerEditor/Source/MediaPlayerEditor/Private/Customizations/MediaSoundWaveCustomization.h
--- a/Engine/Plugins/Media/MediaPlayerEditor/Source/MediaPlayerEditor/Private/Customizations/MediaSoundWaveCustomization.h
+++ b/Engine/Plugins/Media/MediaPlayerEditor/Source/MediaPlayerEditor/Private/Customizations/MediaSoundWaveCustomization.h
@@ -2,6 +2,8 @@
 
 #pragma once
 
+#include "IDetailCustomization.h"
+
 
 /**
  * Implements a details view customization for the UMediaSoundWave class.
@@ -11,9 +13,9 @@ class FMediaSoundWaveCustomization
 {
 public:
 
-	// IDetailCustomization interface
+	//~ IDetailCustomization interface
 
-	virtual void CustomizeDetails( IDetailLayoutBuilder& DetailBuilder ) override;
+	virtual void CustomizeDetails(IDetailLayoutBuilder& DetailBuilder) override;
 
 public:
 
diff --git a/Engine/Plugins/Media/MediaPlayerEditor/Source/MediaPlayerEditor/Private/Customizations/MediaSourceCustomization.cpp b/Engine/Plugins/Media/MediaPlayerEditor/Source/MediaPlayerEditor/Private/Customizations/MediaSourceCustomization.cpp
new file mode 100644
--- /dev/null
+++ b/Engine/Plugins/Media/MediaPlayerEditor/Source/MediaPlayerEditor/Private/Customizations/MediaSourceCustomization.cpp
@@ -0,0 +1,101 @@
+// Copyright 1998-2016 Epic Games, Inc. All Rights Reserved.
+
+#include "MediaPlayerEditorPCH.h"
+#include "MediaSourceCustomization.h"
+
+
+#define LOCTEXT_NAMESPACE "FMediaSourceCustomization"
+
+
+/* IDetailCustomization interface
+ *****************************************************************************/
+
+void FMediaSourceCustomization::CustomizeDetails(IDetailLayoutBuilder& DetailBuilder)
+{
+	// customize 'Platforms' category
+	IDetailCategoryBuilder& PlatformsCategory = DetailBuilder.EditCategory("Platforms");
+	{
+		// DefaultPlayers
+		DefaultPlayersProperty = DetailBuilder.GetProperty("DefaultPlayers");
+		{
+			IDetailPropertyRow& DefaultPlayersRow = PlatformsCategory.AddProperty(DefaultPlayersProperty);
+
+			DefaultPlayersRow.CustomWidget()
+				.NameContent()
+					[
+						DefaultPlayersProperty->CreatePropertyNameWidget()
+					]
+				.ValueContent()
+					[
+						MakeDefaultPlayersValueWidget()
+					];
+		}
+	}
+}
+
+
+/* FMediaSourceCustomization implementation
+ *****************************************************************************/
+
+TSharedRef<SWidget> FMediaSourceCustomization::MakeDefaultPlayersValueWidget() const
+{
+	TArray<const PlatformInfo::FPlatformInfo*> AvailablePlatforms;
+
+	// get available platforms
+	for (const PlatformInfo::FPlatformInfo& PlatformInfo : PlatformInfo::EnumeratePlatformInfoArray())
+	{
+		if (PlatformInfo.IsVanilla() && (PlatformInfo.PlatformType == PlatformInfo::EPlatformType::Game))
+		{
+			// @todo AllDesktop: Re-enable here
+			if (PlatformInfo.PlatformInfoName != TEXT("AllDesktop"))
+			{
+				AvailablePlatforms.Add(&PlatformInfo);
+			}
+		}
+	}
+
+	// sort available platforms alphabetically
+	AvailablePlatforms.Sort([](const PlatformInfo::FPlatformInfo& One, const PlatformInfo::FPlatformInfo& Two) -> bool
+	{
+		return One.DisplayName.CompareTo(Two.DisplayName) < 0;
+	});
+
+	// build value widget
+	TSharedRef<SGridPanel> PlatformPanel = SNew(SGridPanel)
+		.FillColumn(2, 1.0f);
+
+	for (int32 PlatformIndex = 0; PlatformIndex < AvailablePlatforms.Num(); ++PlatformIndex)
+	{
+		const PlatformInfo::FPlatformInfo* Platform = AvailablePlatforms[PlatformIndex];
+
+		// platform icon
+		PlatformPanel->AddSlot(0, PlatformIndex)
+			.VAlign(VAlign_Center)
+			[
+				SNew(SImage)
+					.Image(FEditorStyle::GetBrush(Platform->GetIconStyleName(PlatformInfo::EPlatformIconSize::Normal)))
+			];
+
+		// platform name
+		PlatformPanel->AddSlot(1, PlatformIndex)
+			.Padding(4.0f, 0.0f, 16.0f, 0.0f)
+			.VAlign(VAlign_Center)
+			[
+				SNew(STextBlock)
+					.Text(Platform->DisplayName)
+			];
+
+		// player combo box
+		PlatformPanel->AddSlot(2, PlatformIndex)
+			.VAlign(VAlign_Center)
+			[
+				SNew(SComboBox<TSharedPtr<FString>>)
+				//Platform->PlatformInfoName
+			];
+	}
+
+	return PlatformPanel;
+}
+
+
+#undef LOCTEXT_NAMESPACE
diff --git a/Engine/Plugins/Media/MediaPlayerEditor/Source/MediaPlayerEditor/Private/Customizations/MediaSourceCustomization.h b/Engine/Plugins/Media/MediaPlayerEditor/Source/MediaPlayerEditor/Private/Customizations/MediaSourceCustomization.h
new file mode 100644
--- /dev/null
+++ b/Engine/Plugins/Media/MediaPlayerEditor/Source/MediaPlayerEditor/Private/Customizations/MediaSourceCustomization.h
@@ -0,0 +1,49 @@
+// Copyright 1998-2016 Epic Games, Inc. All Rights Reserved.
+
+#pragma once
+
+#include "IDetailCustomization.h"
+
+
+class IPropertyHandle;
+class SWidget;
+
+
+/**
+ * Implements a details view customization for the UMediaSource class.
+ */
+class FMediaSourceCustomization
+	: public IDetailCustomization
+{
+public:
+
+	//~ IDetailCustomization interface
+
+	virtual void CustomizeDetails(IDetailLayoutBuilder& DetailBuilder) override;
+
+public:
+
+	/**
+	 * Creates an instance of this class.
+	 *
+	 * @return The new instance.
+	 */
+	static TSharedRef<IDetailCustomization> MakeInstance()
+	{
+		return MakeShareable(new FMediaSourceCustomization());
+	}
+
+protected:
+
+	/**
+	 * Makes a widget for the DefaultPlayers property value.
+	 *
+	 * @return The widget.
+	 */
+	TSharedRef<SWidget> MakeDefaultPlayersValueWidget() const;
+
+private:
+
+	/** Pointer to the DefaultPlayers property handle. */
+	TSharedPtr<IPropertyHandle> DefaultPlayersProperty;
+};
diff --git a/Engine/Plugins/Media/MediaPlayerEditor/Source/MediaPlayerEditor/Private/Customizations/MediaTextureCustomization.cpp b/Engine/Plugins/Media/MediaPlayerEditor/Source/MediaPlayerEditor/Private/Customizations/MediaTextureCustomization.cpp
--- a/Engine/Plugins/Media/MediaPlayerEditor/Source/MediaPlayerEditor/Private/Customizations/MediaTextureCustomization.cpp
+++ b/Engine/Plugins/Media/MediaPlayerEditor/Source/MediaPlayerEditor/Private/Customizations/MediaTextureCustomization.cpp
@@ -1,14 +1,7 @@
 // Copyright 1998-2016 Epic Games, Inc. All Rights Reserved.
 
-#include "MediaPlayerEditorPrivatePCH.h"
-#include "IMediaPlayer.h"
-#include "IMediaStream.h"
-#include "IMediaAudioTrack.h"
-#include "IMediaCaptionTrack.h"
-#include "IMediaVideoTrack.h"
-#include "MediaPlayer.h"
+#include "MediaPlayerEditorPCH.h"
 #include "MediaTextureCustomization.h"
-#include "PropertyHandle.h"
 
 
 #define LOCTEXT_NAMESPACE "FMediaTextureCustomization"
@@ -17,8 +10,8 @@
 /* IDetailCustomization interface
  *****************************************************************************/
 
-void FMediaTextureCustomization::CustomizeDetails( IDetailLayoutBuilder& DetailBuilder )
-{
+void FMediaTextureCustomization::CustomizeDetails(IDetailLayoutBuilder& DetailBuilder)
+{/*
 	DetailBuilder.GetObjectsBeingCustomized(CustomizedMediaTextures);
 	MediaPlayerProperty = DetailBuilder.GetProperty("MediaPlayer");
 	VideoTrackIndexProperty = DetailBuilder.GetProperty("VideoTrackIndex");
@@ -48,7 +41,7 @@ void FMediaTextureCustomization::CustomizeDetails( IDetailLayoutBuilder& DetailB
 							.Font(IDetailLayoutBuilder::GetDetailFont())
 					]
 			];
-	}
+	}*/
 }
 
 
@@ -56,7 +49,7 @@ void FMediaTextureCustomization::CustomizeDetails( IDetailLayoutBuilder& DetailB
  *****************************************************************************/
 
 TSharedRef<SWidget> FMediaTextureCustomization::HandleVideoTrackComboButtonMenuContent() const
-{
+{/*
 	// get assigned media player asset
 	UObject* MediaPlayerObj = nullptr;
 	FPropertyAccess::Result Result = MediaPlayerProperty->GetValue(MediaPlayerObj);
@@ -74,7 +67,7 @@ TSharedRef<SWidget> FMediaTextureCustomization::HandleVideoTrackComboButtonMenuC
 	}
 
 	// get media tracks
-	IMediaPlayerPtr Player = MediaPlayer->GetPlayer();
+	TSharedPtr<IMediaPlayer> Player = MediaPlayer->GetPlayer();
 
 	if (!Player.IsValid())
 	{
@@ -98,18 +91,19 @@ TSharedRef<SWidget> FMediaTextureCustomization::HandleVideoTrackComboButtonMenuC
 		);
 	}
 
-	return MenuBuilder.MakeWidget();
+	return MenuBuilder.MakeWidget();*/
+	return SNullWidget::NullWidget;
 }
 
 
-void FMediaTextureCustomization::HandleVideoTrackComboButtonMenuEntryExecute( int32 TrackIndex )
+void FMediaTextureCustomization::HandleVideoTrackComboButtonMenuEntryExecute(int32 TrackIndex)
 {
 	VideoTrackIndexProperty->SetValue(TrackIndex);
 }
 
 
 FText FMediaTextureCustomization::HandleVideoTrackComboButtonText() const
-{
+{/*
 	// get assigned media player asset
 	UObject* MediaPlayerObj = nullptr;
 	FPropertyAccess::Result Result = MediaPlayerProperty->GetValue(MediaPlayerObj);
@@ -136,7 +130,7 @@ FText FMediaTextureCustomization::HandleVideoTrackComboButtonText() const
 	}
 
 	// get selected media track
-	IMediaPlayerPtr Player = MediaPlayer->GetPlayer();
+	TSharedPtr<IMediaPlayer> Player = MediaPlayer->GetPlayer();
 
 	if (!Player.IsValid())
 	{
@@ -151,7 +145,8 @@ FText FMediaTextureCustomization::HandleVideoTrackComboButtonText() const
 	}
 
 	// get track name
-	return VideoTracks[VideoTrackIndex]->GetStream().GetDisplayName();
+	return VideoTracks[VideoTrackIndex]->GetStream().GetDisplayName();*/
+	return FText::GetEmpty();
 }
 
 
diff --git a/Engine/Plugins/Media/MediaPlayerEditor/Source/MediaPlayerEditor/Private/Customizations/MediaTextureCustomization.h b/Engine/Plugins/Media/MediaPlayerEditor/Source/MediaPlayerEditor/Private/Customizations/MediaTextureCustomization.h
--- a/Engine/Plugins/Media/MediaPlayerEditor/Source/MediaPlayerEditor/Private/Customizations/MediaTextureCustomization.h
+++ b/Engine/Plugins/Media/MediaPlayerEditor/Source/MediaPlayerEditor/Private/Customizations/MediaTextureCustomization.h
@@ -2,6 +2,12 @@
 
 #pragma once
 
+#include "IDetailCustomization.h"
+
+
+class IPropertyHandle;
+class SWidget;
+
 
 /**
  * Implements a details view customization for the UMediaTexture class.
@@ -11,7 +17,7 @@ class FMediaTextureCustomization
 {
 public:
 
-	// IDetailCustomization interface
+	//~ IDetailCustomization interface
 
 	virtual void CustomizeDetails(IDetailLayoutBuilder& DetailBuilder) override;
 
diff --git a/Engine/Plugins/Media/MediaPlayerEditor/Source/MediaPlayerEditor/Private/Factories/MediaPlayerFactoryNew.cpp b/Engine/Plugins/Media/MediaPlayerEditor/Source/MediaPlayerEditor/Private/Factories/MediaPlayerFactoryNew.cpp
--- a/Engine/Plugins/Media/MediaPlayerEditor/Source/MediaPlayerEditor/Private/Factories/MediaPlayerFactoryNew.cpp
+++ b/Engine/Plugins/Media/MediaPlayerEditor/Source/MediaPlayerEditor/Private/Factories/MediaPlayerFactoryNew.cpp
@@ -1,6 +1,153 @@
 // Copyright 1998-2016 Epic Games, Inc. All Rights Reserved.
 
-#include "MediaPlayerEditorPrivatePCH.h"
+#include "MediaPlayerEditorPCH.h"
+#include "MediaPlayerFactoryNew.h"
+#include "MediaSoundWaveFactoryNew.h"
+#include "MediaTextureFactoryNew.h"
+
+
+#define LOCTEXT_NAMESPACE "UMediaPlayerFactoryNew"
+
+
+/* Local helpers
+ *****************************************************************************/
+
+class SMediaPlayerFactoryDialog
+	: public SCompoundWidget
+{
+public:
+
+	SLATE_BEGIN_ARGS(SMediaPlayerFactoryDialog) { }
+	SLATE_END_ARGS()
+
+	/** Construct this widget. */
+	void Construct(const FArguments& InArgs, FMediaPlayerFactoryNewOptions& InOptions, TSharedRef<SWindow> InWindow)
+	{
+		Options = &InOptions;
+		Window = InWindow;
+
+		ChildSlot
+		[
+			SNew(SBorder)
+				.Visibility(EVisibility::Visible)
+				.BorderImage(FEditorStyle::GetBrush("Menu.Background"))
+				[
+					SNew(SVerticalBox)
+
+					+ SVerticalBox::Slot()
+						.FillHeight(1)
+						.VAlign(VAlign_Top)
+						[
+							SNew(SBorder)
+								.BorderImage(FEditorStyle::GetBrush("ToolPanel.GroupBorder"))
+								.Padding(4.0f)
+								.Content()
+								[
+									SNew(SVerticalBox)
+
+									+ SVerticalBox::Slot()
+										[
+											SNew(STextBlock)
+												.Text(LOCTEXT("CreateAdditionalAssetsLabel", "Additional assets to create and link to the Media Player:"))
+										]
+
+									+ SVerticalBox::Slot()
+										.Padding(0.0f, 8.0f, 0.0f, 0.0f)
+										[
+											SNew(SCheckBox)
+												.IsChecked(Options->CreateSoundWave ? ECheckBoxState::Checked : ECheckBoxState::Unchecked)
+												.OnCheckStateChanged_Lambda([this](ECheckBoxState CheckBoxState) {
+													Options->CreateSoundWave = (CheckBoxState == ECheckBoxState::Checked);
+												})
+												.Content()
+												[
+													SNew(STextBlock)
+														.Text(LOCTEXT("CreateSoundWaveLabel", "Audio output SoundWave asset"))
+												]
+										]
+
+									+ SVerticalBox::Slot()
+										.Padding(0.0f, 6.0f, 0.0f, 0.0f)
+										[
+											SNew(SCheckBox)
+												.IsChecked(Options->CreateImageTexture ? ECheckBoxState::Checked : ECheckBoxState::Unchecked)
+												.OnCheckStateChanged_Lambda([this](ECheckBoxState CheckBoxState) {
+													Options->CreateImageTexture = (CheckBoxState == ECheckBoxState::Checked);
+												})
+												.Content()
+												[
+													SNew(STextBlock)
+														.Text(LOCTEXT("CreateImageTextureLabel", "Still image output MediaTexture asset"))
+												]
+										]
+
+									+ SVerticalBox::Slot()
+										.Padding(0.0f, 6.0f, 0.0f, 0.0f)
+										[
+											SNew(SCheckBox)
+												.IsChecked(Options->CreateVideoTexture ? ECheckBoxState::Checked : ECheckBoxState::Unchecked)
+												.OnCheckStateChanged_Lambda([this](ECheckBoxState CheckBoxState) {
+													Options->CreateVideoTexture = (CheckBoxState == ECheckBoxState::Checked);
+												})
+												.Content()
+												[
+													SNew(STextBlock)
+														.Text(LOCTEXT("CreateVideoTextureLabel", "Video output MediaTexture asset"))
+												]
+										]
+								]
+						]
+
+					+ SVerticalBox::Slot()
+						.AutoHeight()
+						.HAlign(HAlign_Right)
+						.VAlign(VAlign_Bottom)
+						.Padding(8)
+						[
+							SNew(SUniformGridPanel)
+								.SlotPadding(FEditorStyle::GetMargin("StandardDialog.SlotPadding"))
+								.MinDesiredSlotWidth(FEditorStyle::GetFloat("StandardDialog.MinDesiredSlotWidth"))
+								.MinDesiredSlotHeight(FEditorStyle::GetFloat("StandardDialog.MinDesiredSlotHeight"))
+				
+							+ SUniformGridPanel::Slot(0, 0)
+								[
+									SNew(SButton)
+										.HAlign(HAlign_Center)
+										.ContentPadding(FEditorStyle::GetMargin("StandardDialog.ContentPadding"))
+										.OnClicked_Lambda([this]() -> FReply { CloseDialog(true); return FReply::Handled(); })
+										.Text(LOCTEXT("OkButtonLabel", "OK"))
+								]
+					
+							+ SUniformGridPanel::Slot(1, 0)
+								[
+									SNew(SButton)
+										.HAlign(HAlign_Center)
+										.ContentPadding(FEditorStyle::GetMargin("StandardDialog.ContentPadding"))
+										.OnClicked_Lambda([this]() -> FReply { CloseDialog(false); return FReply::Handled(); })
+										.Text(LOCTEXT("CancelButtonLabel", "Cancel"))
+								]
+						]
+				]
+		];
+	}
+
+protected:
+
+	void CloseDialog(bool InOkClicked)
+	{
+		Options->OkClicked = InOkClicked;
+
+		if (Window.IsValid())
+		{
+			Window.Pin()->RequestDestroyWindow();
+		}
+	}
+
+private:
+
+	FMediaPlayerFactoryNewOptions* Options;
+	TWeakPtr<SWindow> Window;
+};
 
 
 /* UMediaPlayerFactoryNew structors
@@ -15,12 +162,88 @@ UMediaPlayerFactoryNew::UMediaPlayerFactoryNew( const FObjectInitializer& Object
 }
 
 
-/* UFactory overrides
+/* UFactory interface
  *****************************************************************************/
 
-UObject* UMediaPlayerFactoryNew::FactoryCreateNew( UClass* InClass, UObject* InParent, FName InName, EObjectFlags Flags, UObject* Context, FFeedbackContext* Warn )
+bool UMediaPlayerFactoryNew::ConfigureProperties()
 {
-	return NewObject<UMediaPlayer>(InParent, InClass, InName, Flags);
+	Options.CreateImageTexture = false;
+	Options.CreateSoundWave = false;
+	Options.CreateVideoTexture = false;
+	Options.OkClicked = false;
+
+	TSharedRef<SWindow> Window = SNew(SWindow)
+		.Title(LOCTEXT("CreateMediaPlayerDialogTitle", "Create Media Player"))
+		.ClientSize(FVector2D(400, 160))
+		.SupportsMinimize(false)
+		.SupportsMaximize(false);
+
+	Window->SetContent(SNew(SMediaPlayerFactoryDialog, Options, Window));
+	GEditor->EditorAddModalWindow(Window);
+
+	return Options.OkClicked;
+}
+
+
+UObject* UMediaPlayerFactoryNew::FactoryCreateNew(UClass* InClass, UObject* InParent, FName InName, EObjectFlags Flags, UObject* Context, FFeedbackContext* Warn)
+{
+	auto NewMediaPlayer = NewObject<UMediaPlayer>(InParent, InClass, InName, Flags);
+
+	if ((NewMediaPlayer != nullptr) && (Options.CreateImageTexture || Options.CreateSoundWave || Options.CreateVideoTexture))
+	{
+		IAssetTools& AssetTools = FModuleManager::Get().LoadModuleChecked<FAssetToolsModule>("AssetTools").Get();
+		const FString ParentName = InParent->GetOutermost()->GetName();
+
+		FString OutAssetName;
+		FString OutPackageName;
+
+		if (Options.CreateImageTexture)
+		{
+			AssetTools.CreateUniqueAssetName(ParentName, TEXT("_Image"), OutPackageName, OutAssetName);
+			const FString PackagePath = FPackageName::GetLongPackagePath(OutPackageName);
+			auto Factory = NewObject<UMediaTextureFactoryNew>();
+			auto ImageTexture = Cast<UMediaTexture>(AssetTools.CreateAsset(OutAssetName, PackagePath, UMediaTexture::StaticClass(), Factory));
+
+			if (ImageTexture != nullptr)
+			{
+				NewMediaPlayer->SetImageTexture(ImageTexture);
+			}
+		}
+
+		if (Options.CreateSoundWave)
+		{
+			AssetTools.CreateUniqueAssetName(ParentName, TEXT("_Sound"), OutPackageName, OutAssetName);
+			const FString PackagePath = FPackageName::GetLongPackagePath(OutPackageName);
+			auto Factory = NewObject<UMediaSoundWaveFactoryNew>();
+			auto SoundWave = Cast<UMediaSoundWave>(AssetTools.CreateAsset(OutAssetName, PackagePath, UMediaSoundWave::StaticClass(), Factory));
+
+			if (SoundWave != nullptr)
+			{
+				NewMediaPlayer->SetSoundWave(SoundWave);
+			}
+		}
+
+		if (Options.CreateVideoTexture)
+		{
+			AssetTools.CreateUniqueAssetName(ParentName, TEXT("_Video"), OutPackageName, OutAssetName);
+			const FString PackagePath = FPackageName::GetLongPackagePath(OutPackageName);
+			auto Factory = NewObject<UMediaTextureFactoryNew>();
+			auto VideoTexture = Cast<UMediaTexture>(AssetTools.CreateAsset(OutAssetName, PackagePath, UMediaTexture::StaticClass(), Factory));
+
+			if (VideoTexture != nullptr)
+			{
+				NewMediaPlayer->SetVideoTexture(VideoTexture);
+			}
+		}
+	}
+
+	return NewMediaPlayer;
+}
+
+
+uint32 UMediaPlayerFactoryNew::GetMenuCategories() const
+{
+	return EAssetTypeCategories::Media;
 }
 
 
@@ -28,3 +251,6 @@ bool UMediaPlayerFactoryNew::ShouldShowInNewMenu() const
 {
 	return true;
 }
+
+
+#undef LOCTEXT_NAMESPACE
diff --git a/Engine/Plugins/Media/MediaPlayerEditor/Source/MediaPlayerEditor/Private/Factories/MediaPlayerFactoryNew.h b/Engine/Plugins/Media/MediaPlayerEditor/Source/MediaPlayerEditor/Private/Factories/MediaPlayerFactoryNew.h
--- a/Engine/Plugins/Media/MediaPlayerEditor/Source/MediaPlayerEditor/Private/Factories/MediaPlayerFactoryNew.h
+++ b/Engine/Plugins/Media/MediaPlayerEditor/Source/MediaPlayerEditor/Private/Factories/MediaPlayerFactoryNew.h
@@ -6,6 +6,18 @@
 #include "MediaPlayerFactoryNew.generated.h"
 
 
+/**
+ * Options for UMediaPlayerFactoryNew.
+ */
+struct FMediaPlayerFactoryNewOptions
+{
+	bool CreateImageTexture;
+	bool CreateSoundWave;
+	bool CreateVideoTexture;
+	bool OkClicked;
+};
+
+
 /**
  * Implements a factory for UMediaPlayer objects.
  */
@@ -17,8 +29,14 @@ class UMediaPlayerFactoryNew
 
 public:
 
-	// UFactory Interface
+	//~ UFactory Interface
 
-	virtual UObject* FactoryCreateNew( UClass* InClass, UObject* InParent, FName InName, EObjectFlags Flags, UObject* Context, FFeedbackContext* Warn ) override;
+	virtual bool ConfigureProperties() override;
+	virtual UObject* FactoryCreateNew(UClass* InClass, UObject* InParent, FName InName, EObjectFlags Flags, UObject* Context, FFeedbackContext* Warn) override;
+	virtual uint32 GetMenuCategories() const override;
 	virtual bool ShouldShowInNewMenu() const override;
+
+private:
+
+	FMediaPlayerFactoryNewOptions Options;
 };
diff --git a/Engine/Plugins/Media/MediaPlayerEditor/Source/MediaPlayerEditor/Private/Factories/MediaSoundWaveFactoryNew.cpp b/Engine/Plugins/Media/MediaPlayerEditor/Source/MediaPlayerEditor/Private/Factories/MediaSoundWaveFactoryNew.cpp
--- a/Engine/Plugins/Media/MediaPlayerEditor/Source/MediaPlayerEditor/Private/Factories/MediaSoundWaveFactoryNew.cpp
+++ b/Engine/Plugins/Media/MediaPlayerEditor/Source/MediaPlayerEditor/Private/Factories/MediaSoundWaveFactoryNew.cpp
@@ -1,12 +1,13 @@
 // Copyright 1998-2016 Epic Games, Inc. All Rights Reserved.
 
-#include "MediaPlayerEditorPrivatePCH.h"
+#include "MediaPlayerEditorPCH.h"
+#include "MediaSoundWaveFactoryNew.h"
 
 
 /* UMediaSoundWaveFactoryNew structors
  *****************************************************************************/
 
-UMediaSoundWaveFactoryNew::UMediaSoundWaveFactoryNew( const FObjectInitializer& ObjectInitializer )
+UMediaSoundWaveFactoryNew::UMediaSoundWaveFactoryNew(const FObjectInitializer& ObjectInitializer)
 	: Super(ObjectInitializer)
 {
 	SupportedClass = UMediaSoundWave::StaticClass();
@@ -15,19 +16,18 @@ UMediaSoundWaveFactoryNew::UMediaSoundWaveFactoryNew( const FObjectInitializer&
 }
 
 
-/* UFactory overrides
+/* UFactory interface
  *****************************************************************************/
 
-UObject* UMediaSoundWaveFactoryNew::FactoryCreateNew( UClass* InClass, UObject* InParent, FName InName, EObjectFlags Flags, UObject* Context, FFeedbackContext* Warn )
+UObject* UMediaSoundWaveFactoryNew::FactoryCreateNew(UClass* InClass, UObject* InParent, FName InName, EObjectFlags Flags, UObject* Context, FFeedbackContext* Warn)
 {
-	UMediaSoundWave* MediaSoundWave = NewObject<UMediaSoundWave>(InParent, InClass, InName, Flags);
+	return NewObject<UMediaSoundWave>(InParent, InClass, InName, Flags);
+}
 
-	if ((MediaSoundWave != nullptr) && (InitialMediaPlayer != nullptr))
-	{
-		MediaSoundWave->SetMediaPlayer(InitialMediaPlayer);
-	}
 
-	return MediaSoundWave;
+uint32 UMediaSoundWaveFactoryNew::GetMenuCategories() const
+{
+	return EAssetTypeCategories::Media | EAssetTypeCategories::Sounds;
 }
 
 
diff --git a/Engine/Plugins/Media/MediaPlayerEditor/Source/MediaPlayerEditor/Private/Factories/MediaSoundWaveFactoryNew.h b/Engine/Plugins/Media/MediaPlayerEditor/Source/MediaPlayerEditor/Private/Factories/MediaSoundWaveFactoryNew.h
--- a/Engine/Plugins/Media/MediaPlayerEditor/Source/MediaPlayerEditor/Private/Factories/MediaSoundWaveFactoryNew.h
+++ b/Engine/Plugins/Media/MediaPlayerEditor/Source/MediaPlayerEditor/Private/Factories/MediaSoundWaveFactoryNew.h
@@ -2,6 +2,7 @@
 
 #pragma once
 
+#include "Factories/Factory.h"
 #include "MediaSoundWaveFactoryNew.generated.h"
 
 
@@ -14,14 +15,11 @@ class UMediaSoundWaveFactoryNew
 {
 	GENERATED_UCLASS_BODY()
 
-	/** An initial UMediaPlayer asset to place in the newly created sound wave. */
-	UPROPERTY()
-	class UMediaPlayer* InitialMediaPlayer;
-
 public:
 
-	// UFactory Interface
+	//~ UFactory Interface
 
-	virtual UObject* FactoryCreateNew( UClass* InClass, UObject* InParent, FName InName, EObjectFlags Flags, UObject* Context, FFeedbackContext* Warn ) override;
+	virtual UObject* FactoryCreateNew(UClass* InClass, UObject* InParent, FName InName, EObjectFlags Flags, UObject* Context, FFeedbackContext* Warn) override;
+	virtual uint32 GetMenuCategories() const override;
 	virtual bool ShouldShowInNewMenu() const override;
 };
diff --git a/Engine/Plugins/Media/MediaPlayerEditor/Source/MediaPlayerEditor/Private/Factories/MediaTextureFactoryNew.cpp b/Engine/Plugins/Media/MediaPlayerEditor/Source/MediaPlayerEditor/Private/Factories/MediaTextureFactoryNew.cpp
--- a/Engine/Plugins/Media/MediaPlayerEditor/Source/MediaPlayerEditor/Private/Factories/MediaTextureFactoryNew.cpp
+++ b/Engine/Plugins/Media/MediaPlayerEditor/Source/MediaPlayerEditor/Private/Factories/MediaTextureFactoryNew.cpp
@@ -1,12 +1,13 @@
 // Copyright 1998-2016 Epic Games, Inc. All Rights Reserved.
 
-#include "MediaPlayerEditorPrivatePCH.h"
+#include "MediaPlayerEditorPCH.h"
+#include "MediaTextureFactoryNew.h"
 
 
 /* UMediaTextureFactoryNew structors
  *****************************************************************************/
 
-UMediaTextureFactoryNew::UMediaTextureFactoryNew( const FObjectInitializer& ObjectInitializer )
+UMediaTextureFactoryNew::UMediaTextureFactoryNew(const FObjectInitializer& ObjectInitializer)
 	: Super(ObjectInitializer)
 {
 	SupportedClass = UMediaTexture::StaticClass();
@@ -18,19 +19,25 @@ UMediaTextureFactoryNew::UMediaTextureFactoryNew( const FObjectInitializer& Obje
 /* UFactory overrides
  *****************************************************************************/
 
-UObject* UMediaTextureFactoryNew::FactoryCreateNew( UClass* InClass, UObject* InParent, FName InName, EObjectFlags Flags, UObject* Context, FFeedbackContext* Warn )
+UObject* UMediaTextureFactoryNew::FactoryCreateNew(UClass* InClass, UObject* InParent, FName InName, EObjectFlags Flags, UObject* Context, FFeedbackContext* Warn)
 {
-	UMediaTexture* MediaTexture = NewObject<UMediaTexture>(InParent, InClass, InName, Flags);
+	auto MediaTexture = NewObject<UMediaTexture>(InParent, InClass, InName, Flags);
 
-	if ((MediaTexture != nullptr) && (InitialMediaPlayer != nullptr))
+	if (MediaTexture != nullptr)
 	{
-		MediaTexture->SetMediaPlayer(InitialMediaPlayer);
+		MediaTexture->UpdateResource();
 	}
 
 	return MediaTexture;
 }
 
 
+uint32 UMediaTextureFactoryNew::GetMenuCategories() const
+{
+	return EAssetTypeCategories::Media | EAssetTypeCategories::MaterialsAndTextures;
+}
+
+
 bool UMediaTextureFactoryNew::ShouldShowInNewMenu() const
 {
 	return true;
diff --git a/Engine/Plugins/Media/MediaPlayerEditor/Source/MediaPlayerEditor/Private/Factories/MediaTextureFactoryNew.h b/Engine/Plugins/Media/MediaPlayerEditor/Source/MediaPlayerEditor/Private/Factories/MediaTextureFactoryNew.h
--- a/Engine/Plugins/Media/MediaPlayerEditor/Source/MediaPlayerEditor/Private/Factories/MediaTextureFactoryNew.h
+++ b/Engine/Plugins/Media/MediaPlayerEditor/Source/MediaPlayerEditor/Private/Factories/MediaTextureFactoryNew.h
@@ -2,6 +2,7 @@
 
 #pragma once
 
+#include "Factories/Factory.h"
 #include "MediaTextureFactoryNew.generated.h"
 
 
@@ -14,14 +15,11 @@ class UMediaTextureFactoryNew
 {
 	GENERATED_UCLASS_BODY()
 
-	/** An initial media player asset to place in the newly created texture. */
-	UPROPERTY()
-	class UMediaPlayer* InitialMediaPlayer;
-
 public:
 
-	// UFactory Interface
+	//~ UFactory Interface
 
 	virtual UObject* FactoryCreateNew( UClass* InClass, UObject* InParent, FName InName, EObjectFlags Flags, UObject* Context, FFeedbackContext* Warn ) override;
+	virtual uint32 GetMenuCategories() const override;
 	virtual bool ShouldShowInNewMenu() const override;
 };
diff --git a/Engine/Plugins/Media/MediaPlayerEditor/Source/MediaPlayerEditor/Private/MediaPlayerEditorClasses.cpp b/Engine/Plugins/Media/MediaPlayerEditor/Source/MediaPlayerEditor/Private/MediaPlayerEditorClasses.cpp
new file mode 100644
--- /dev/null
+++ b/Engine/Plugins/Media/MediaPlayerEditor/Source/MediaPlayerEditor/Private/MediaPlayerEditorClasses.cpp
@@ -0,0 +1,11 @@
+// Copyright 1998-2016 Epic Games, Inc. All Rights Reserved.
+
+#include "MediaPlayerEditorPCH.h"
+#include "MediaPlayerEditorSettings.h"
+
+
+UMediaPlayerEditorSettings::UMediaPlayerEditorSettings(const FObjectInitializer& ObjectInitializer)
+	: Super(ObjectInitializer)
+	, ScaleVideoToFit(true)
+	, ShowCaptions(true)
+{ }
diff --git a/Engine/Plugins/Media/MediaPlayerEditor/Source/MediaPlayerEditor/Private/MediaPlayerEditorModule.cpp b/Engine/Plugins/Media/MediaPlayerEditor/Source/MediaPlayerEditor/Private/MediaPlayerEditorModule.cpp
--- a/Engine/Plugins/Media/MediaPlayerEditor/Source/MediaPlayerEditor/Private/MediaPlayerEditorModule.cpp
+++ b/Engine/Plugins/Media/MediaPlayerEditor/Source/MediaPlayerEditor/Private/MediaPlayerEditorModule.cpp
@@ -1,12 +1,24 @@
 // Copyright 1998-2016 Epic Games, Inc. All Rights Reserved.
 
-
-#include "MediaPlayerEditorPrivatePCH.h"
-#include "ModuleInterface.h"
+#include "MediaPlayerEditorPCH.h"
+#include "FileMediaSourceActions.h"
+#include "FileMediaSourceCustomization.h"
+#include "MediaPlayerActions.h"
+#include "MediaPlayerEditorCommands.h"
+#include "MediaPlayerEditorStyle.h"
+#include "MediaPlaylistActions.h"
+#include "MediaSoundWaveActions.h"
+#include "MediaSoundWaveCustomization.h"
+#include "MediaSourceActions.h"
+#include "MediaSourceCustomization.h"
+#include "MediaTextureActions.h"
+#include "MediaTextureCustomization.h"
 
 
 #define LOCTEXT_NAMESPACE "FMediaPlayerEditorModule"
 
+DEFINE_LOG_CATEGORY(LogMediaPlayerEditor);
+
 
 /**
  * Implements the MediaPlayerEditor module.
@@ -18,7 +30,7 @@ class FMediaPlayerEditorModule
 {
 public:
 
-	// IHasMenuExtensibility interface
+	//~ IHasMenuExtensibility interface
 
 	virtual TSharedPtr<FExtensibilityManager> GetMenuExtensibilityManager() override
 	{
@@ -27,7 +39,7 @@ public:
 
 public:
 
-	// IHasToolBarExtensibility interface
+	//~ IHasToolBarExtensibility interface
 
 	virtual TSharedPtr<FExtensibilityManager> GetToolBarExtensibilityManager() override
 	{
@@ -36,7 +48,7 @@ public:
 
 public:
 
-	// IModuleInterface interface
+	//~ IModuleInterface interface
 
 	virtual void StartupModule() override
 	{
@@ -46,6 +58,7 @@ public:
 
 		RegisterAssetTools();
 		RegisterCustomizations();
+		RegisterEditorDelegates();
 		RegisterMenuExtensions();
 		RegisterThumbnailRenderers();
 	}
@@ -54,6 +67,7 @@ public:
 	{
 		UnregisterAssetTools();
 		UnregisterCustomizations();
+		UnregisterEditorDelegates();
 		UnregisterMenuExtensions();
 		UnregisterThumbnailRenderers();
 	}
@@ -65,8 +79,11 @@ protected:
 	{
 		IAssetTools& AssetTools = FModuleManager::LoadModuleChecked<FAssetToolsModule>("AssetTools").Get();
 
+		RegisterAssetTypeAction(AssetTools, MakeShareable(new FFileMediaSourceActions(Style.ToSharedRef())));
 		RegisterAssetTypeAction(AssetTools, MakeShareable(new FMediaPlayerActions(Style.ToSharedRef())));
+		RegisterAssetTypeAction(AssetTools, MakeShareable(new FMediaPlaylistActions(Style.ToSharedRef())));
 		RegisterAssetTypeAction(AssetTools, MakeShareable(new FMediaSoundWaveActions));
+		RegisterAssetTypeAction(AssetTools, MakeShareable(new FMediaSourceActions));
 		RegisterAssetTypeAction(AssetTools, MakeShareable(new FMediaTextureActions));
 	}
 
@@ -103,22 +120,50 @@ protected:
 	/** Register details view customizations. */
 	void RegisterCustomizations()
 	{
-		FPropertyEditorModule& PropertyModule = FModuleManager::LoadModuleChecked<FPropertyEditorModule>("PropertyEditor");
-		MediaPlayerName    = UMediaPlayer   ::StaticClass()->GetFName();
+		FileMediaSourceName = UFileMediaSource::StaticClass()->GetFName();
 		MediaSoundWaveName = UMediaSoundWave::StaticClass()->GetFName();
-		MediaTextureName   = UMediaTexture  ::StaticClass()->GetFName();
-		PropertyModule.RegisterCustomClassLayout(MediaPlayerName, FOnGetDetailCustomizationInstance::CreateStatic(&FMediaPlayerCustomization::MakeInstance));
-		PropertyModule.RegisterCustomClassLayout(MediaSoundWaveName, FOnGetDetailCustomizationInstance::CreateStatic(&FMediaSoundWaveCustomization::MakeInstance));
-		PropertyModule.RegisterCustomClassLayout(MediaTextureName, FOnGetDetailCustomizationInstance::CreateStatic(&FMediaTextureCustomization::MakeInstance));
+		MediaSourceName = UMediaSource::StaticClass()->GetFName();
+		MediaTextureName = UMediaTexture::StaticClass()->GetFName();
+
+		FPropertyEditorModule& PropertyModule = FModuleManager::LoadModuleChecked<FPropertyEditorModule>("PropertyEditor");
+		{
+			PropertyModule.RegisterCustomClassLayout(FileMediaSourceName, FOnGetDetailCustomizationInstance::CreateStatic(&FFileMediaSourceCustomization::MakeInstance));
+			PropertyModule.RegisterCustomClassLayout(MediaSoundWaveName, FOnGetDetailCustomizationInstance::CreateStatic(&FMediaSoundWaveCustomization::MakeInstance));
+			PropertyModule.RegisterCustomClassLayout(MediaSourceName, FOnGetDetailCustomizationInstance::CreateStatic(&FMediaSourceCustomization::MakeInstance));
+			PropertyModule.RegisterCustomClassLayout(MediaTextureName, FOnGetDetailCustomizationInstance::CreateStatic(&FMediaTextureCustomization::MakeInstance));
+		}
 	}
 
 	/** Unregister details view customizations. */
 	void UnregisterCustomizations()
 	{
 		FPropertyEditorModule& PropertyModule = FModuleManager::LoadModuleChecked<FPropertyEditorModule>("PropertyEditor");
-		PropertyModule.UnregisterCustomClassLayout(MediaPlayerName);
-		PropertyModule.UnregisterCustomClassLayout(MediaSoundWaveName);
-		PropertyModule.UnregisterCustomClassLayout(MediaTextureName);
+		{
+			PropertyModule.UnregisterCustomClassLayout(FileMediaSourceName);
+			PropertyModule.UnregisterCustomClassLayout(MediaSoundWaveName);
+			PropertyModule.UnregisterCustomClassLayout(MediaSourceName);
+			PropertyModule.UnregisterCustomClassLayout(MediaTextureName);
+		}
+	}
+
+protected:
+
+	/** Register Editor delegates. */
+	void RegisterEditorDelegates()
+	{
+		FEditorDelegates::BeginPIE.AddRaw(this, &FMediaPlayerEditorModule::HandleEditorBeginPIE);
+		FEditorDelegates::EndPIE.AddRaw(this, &FMediaPlayerEditorModule::HandleEditorEndPIE);
+		FEditorDelegates::PausePIE.AddRaw(this, &FMediaPlayerEditorModule::HandleEditorPausePIE);
+		FEditorDelegates::ResumePIE.AddRaw(this, &FMediaPlayerEditorModule::HandleEditorResumePIE);
+	}
+
+	/** Unregister Editor delegates. */
+	void UnregisterEditorDelegates()
+	{
+		FEditorDelegates::BeginPIE.RemoveAll(this);
+		FEditorDelegates::EndPIE.RemoveAll(this);
+		FEditorDelegates::PausePIE.RemoveAll(this);
+		FEditorDelegates::ResumePIE.RemoveAll(this);
 	}
 
 protected:
@@ -154,6 +199,40 @@ protected:
 		}
 	}
 
+private:
+
+	void HandleEditorBeginPIE(bool bIsSimulating)
+	{
+		for (TObjectIterator<UMediaPlayer> It; It; ++It)
+		{
+			(*It)->Close();
+		}
+	}
+
+	void HandleEditorEndPIE(bool bIsSimulating)
+	{
+		for (TObjectIterator<UMediaPlayer> It; It; ++It)
+		{
+			(*It)->Close();
+		}
+	}
+
+	void HandleEditorPausePIE(bool bIsSimulating)
+	{
+		for (TObjectIterator<UMediaPlayer> It; It; ++It)
+		{
+			(*It)->PausePIE();
+		}
+	}
+
+	void HandleEditorResumePIE(bool bIsSimulating)
+	{
+		for (TObjectIterator<UMediaPlayer> It; It; ++It)
+		{
+			(*It)->ResumePIE();
+		}
+	}
+
 private:
 
 	/** Holds the menu extensibility manager. */
@@ -168,9 +247,10 @@ private:
 	/** Holds the tool bar extensibility manager. */
 	TSharedPtr<FExtensibilityManager> ToolBarExtensibilityManager;
 
-	/** Class names */
-	FName MediaPlayerName;
+	/** Class names. */
+	FName FileMediaSourceName;
 	FName MediaSoundWaveName;
+	FName MediaSourceName;
 	FName MediaTextureName;
 };
 
diff --git a/Engine/Plugins/Media/MediaPlayerEditor/Source/MediaPlayerEditor/Private/MediaPlayerEditorPCH.h b/Engine/Plugins/Media/MediaPlayerEditor/Source/MediaPlayerEditor/Private/MediaPlayerEditorPCH.h
new file mode 100644
--- /dev/null
+++ b/Engine/Plugins/Media/MediaPlayerEditor/Source/MediaPlayerEditor/Private/MediaPlayerEditorPCH.h
@@ -0,0 +1,63 @@
+// Copyright 1998-2016 Epic Games, Inc. All Rights Reserved.
+
+#pragma once
+
+#include "Developer/AssetTools/Public/AssetToolsModule.h"
+#include "Developer/DesktopPlatform/Public/PlatformInfo.h"
+#include "Developer/DesktopWidgets/Public/Widgets/Input/SFilePathPicker.h"
+
+#include "Editor/ContentBrowser/Public/ContentBrowserModule.h"
+#include "Editor/EditorStyle/Public/EditorStyle.h"
+#include "Editor/PropertyEditor/Public/DetailCategoryBuilder.h"
+#include "Editor/PropertyEditor/Public/DetailLayoutBuilder.h"
+#include "Editor/PropertyEditor/Public/DetailWidgetRow.h"
+#include "Editor/PropertyEditor/Public/IDetailPropertyRow.h"
+#include "Editor/PropertyEditor/Public/IDetailsView.h"
+#include "Editor/PropertyEditor/Public/PropertyEditorModule.h"
+#include "Editor/PropertyEditor/Public/PropertyHandle.h"
+#include "Editor/TextureEditor/Public/TextureEditor.h"
+#include "Editor/UnrealEd/Public/Factories.h"
+#include "Editor/UnrealEd/Public/ScopedTransaction.h"
+#include "Editor/UnrealEd/Public/UnrealEd.h"
+#include "Editor/UnrealEd/Public/Toolkits/AssetEditorManager.h"
+#include "Editor/UnrealEd/Public/Toolkits/AssetEditorToolkit.h"
+#include "Editor/UnrealEd/Public/Toolkits/IToolkit.h"
+#include "Editor/WorkspaceMenuStructure/Public/WorkspaceMenuStructureModule.h"
+
+#include "Runtime/Core/Public/Core.h"
+#include "Runtime/Core/Public/Modules/ModuleInterface.h"
+#include "Runtime/Core/Public/Modules/ModuleManager.h"
+#include "Runtime/CoreUObject/Public/CoreUObject.h"
+#include "Runtime/Engine/Classes/Components/AudioComponent.h"
+#include "Runtime/Engine/Classes/Materials/Material.h"
+#include "Runtime/Engine/Classes/Materials/MaterialExpressionTextureSample.h"
+#include "Runtime/Engine/Public/AudioDevice.h"
+#include "Runtime/Engine/Public/Engine.h"
+#include "Runtime/Engine/Public/Slate/SceneViewport.h"
+#include "Runtime/Engine/Public/Slate/SlateTextures.h"
+#include "Runtime/Media/Public/IMediaControls.h"
+#include "Runtime/Media/Public/IMediaModule.h"
+#include "Runtime/Media/Public/IMediaOutput.h"
+#include "Runtime/Media/Public/IMediaPlayer.h"
+#include "Runtime/Media/Public/IMediaPlayerFactory.h"
+#include "Runtime/Media/Public/IMediaTracks.h"
+#include "Runtime/MediaAssets/Public/FileMediaSource.h"
+#include "Runtime/MediaAssets/Public/MediaPlayer.h"
+#include "Runtime/MediaAssets/Public/MediaPlaylist.h"
+#include "Runtime/MediaAssets/Public/MediaSoundWave.h"
+#include "Runtime/MediaAssets/Public/MediaSource.h"
+#include "Runtime/MediaAssets/Public/MediaTexture.h"
+#include "Runtime/MediaAssets/Public/PlatformMediaSource.h"
+#include "Runtime/MediaAssets/Public/StreamMediaSource.h"
+#include "Runtime/RenderCore/Public/TickableObjectRenderThread.h"
+#include "Runtime/Slate/Public/SlateBasics.h"
+#include "Runtime/Slate/Public/Framework/Notifications/NotificationManager.h"
+#include "Runtime/Slate/Public/Widgets/Docking/SDockTab.h"
+#include "Runtime/Slate/Public/Widgets/Input/SNumericEntryBox.h"
+#include "Runtime/Slate/Public/Widgets/Layout/SScaleBox.h"
+#include "Runtime/Slate/Public/Widgets/Notifications/SNotificationList.h"
+#include "Runtime/SlateCore/Public/Styling/SlateStyle.h"
+
+
+/** Log category for the MediaPlayerEditor module. */
+DECLARE_LOG_CATEGORY_EXTERN(LogMediaPlayerEditor, Log, All);
diff --git a/Engine/Plugins/Media/MediaPlayerEditor/Source/MediaPlayerEditor/Private/Models/MediaPlayerEditorCommands.cpp b/Engine/Plugins/Media/MediaPlayerEditor/Source/MediaPlayerEditor/Private/Models/MediaPlayerEditorCommands.cpp
--- a/Engine/Plugins/Media/MediaPlayerEditor/Source/MediaPlayerEditor/Private/Models/MediaPlayerEditorCommands.cpp
+++ b/Engine/Plugins/Media/MediaPlayerEditor/Source/MediaPlayerEditor/Private/Models/MediaPlayerEditorCommands.cpp
@@ -1,16 +1,23 @@
 // Copyright 1998-2016 Epic Games, Inc. All Rights Reserved.
 
-#include "MediaPlayerEditorPrivatePCH.h"
+#include "MediaPlayerEditorPCH.h"
+#include "MediaPlayerEditorCommands.h"
+
+
+#define LOCTEXT_NAMESPACE "FMediaPlayerEditorCommands"
 
-#define LOCTEXT_NAMESPACE "MediaPlayerEditorCommands"
 
 void FMediaPlayerEditorCommands::RegisterCommands()
 {
-	UI_COMMAND(ForwardMedia, "Forward", "Fast forwards media playback", EUserInterfaceActionType::Button, FInputChord());
+	UI_COMMAND(CloseMedia, "Close", "Close the currently opened media", EUserInterfaceActionType::Button, FInputChord());
+	UI_COMMAND(ForwardMedia, "Forward", "Fast forward media playback", EUserInterfaceActionType::Button, FInputChord());
+	UI_COMMAND(NextMedia, "Next", "Jump to next item in the play list", EUserInterfaceActionType::Button, FInputChord());
 	UI_COMMAND(PauseMedia, "Pause", "Pause media playback", EUserInterfaceActionType::Button, FInputChord());
 	UI_COMMAND(PlayMedia, "Play", "Start media playback", EUserInterfaceActionType::Button, FInputChord());
-	UI_COMMAND(ReverseMedia, "Reverse", "Reverses media playback", EUserInterfaceActionType::Button, FInputChord());
-	UI_COMMAND(RewindMedia, "Rewind", "Rewinds the media to the beginning", EUserInterfaceActionType::Button, FInputChord());
+	UI_COMMAND(PreviousMedia, "Prev", "Jump to previous item in the play list", EUserInterfaceActionType::Button, FInputChord());
+	UI_COMMAND(ReverseMedia, "Reverse", "Reverse media playback", EUserInterfaceActionType::Button, FInputChord());
+	UI_COMMAND(RewindMedia, "Rewind", "Rewind the media to the beginning", EUserInterfaceActionType::Button, FInputChord());
 }
 
-#undef LOCTEXT_NAMESPACE
\ No newline at end of file
+
+#undef LOCTEXT_NAMESPACE
diff --git a/Engine/Plugins/Media/MediaPlayerEditor/Source/MediaPlayerEditor/Private/Shared/MediaPlayerEditorSettings.h b/Engine/Plugins/Media/MediaPlayerEditor/Source/MediaPlayerEditor/Private/Shared/MediaPlayerEditorSettings.h
new file mode 100644
--- /dev/null
+++ b/Engine/Plugins/Media/MediaPlayerEditor/Source/MediaPlayerEditor/Private/Shared/MediaPlayerEditorSettings.h
@@ -0,0 +1,27 @@
+// Copyright 1998-2016 Epic Games, Inc. All Rights Reserved.
+
+#pragma once
+
+#include "MediaPlayerEditorSettings.generated.h"
+
+
+UCLASS(config=EditorPerProjectUserSettings)
+class UMediaPlayerEditorSettings
+	: public UObject
+{
+	GENERATED_UCLASS_BODY()
+
+public:
+
+	/** The name of the desired native media player to use for playback. */
+	UPROPERTY(config, EditAnywhere, Category=Viewer)
+	FName DesiredPlayerName;
+
+	/** Whether the media texture should be scaled to maintain the video's aspect ratio. */
+	UPROPERTY(config, EditAnywhere, Category=Viewer)
+	bool ScaleVideoToFit;
+
+	/** Whether to display caption/subtitle text. */
+	UPROPERTY(config, EditAnywhere, Category=Viewer)
+	bool ShowCaptions;
+};
diff --git a/Engine/Plugins/Media/MediaPlayerEditor/Source/MediaPlayerEditor/Private/Shared/MediaPlayerEditorStyle.h b/Engine/Plugins/Media/MediaPlayerEditor/Source/MediaPlayerEditor/Private/Shared/MediaPlayerEditorStyle.h
new file mode 100644
--- /dev/null
+++ b/Engine/Plugins/Media/MediaPlayerEditor/Source/MediaPlayerEditor/Private/Shared/MediaPlayerEditorStyle.h
@@ -0,0 +1,83 @@
+// Copyright 1998-2016 Epic Games, Inc. All Rights Reserved.
+
+#pragma once
+
+
+#define IMAGE_BRUSH(RelativePath, ...) FSlateImageBrush(RootToContentDir(RelativePath, TEXT(".png")), __VA_ARGS__)
+#define BOX_BRUSH(RelativePath, ...) FSlateBoxBrush(RootToContentDir(RelativePath, TEXT(".png")), __VA_ARGS__)
+#define BORDER_BRUSH(RelativePath, ...) FSlateBorderBrush(RootToContentDir(RelativePath, TEXT(".png")), __VA_ARGS__)
+#define TTF_FONT(RelativePath, ...) FSlateFontInfo(RootToContentDir(RelativePath, TEXT(".ttf")), __VA_ARGS__)
+#define OTF_FONT(RelativePath, ...) FSlateFontInfo(RootToContentDir(RelativePath, TEXT(".otf")), __VA_ARGS__)
+
+
+/**
+ * Implements the visual style of the media player editor UI.
+ */
+class FMediaPlayerEditorStyle
+	: public FSlateStyleSet
+{
+public:
+
+	/** Default constructor. */
+	FMediaPlayerEditorStyle()
+		: FSlateStyleSet("MediaPlayerEditorStyle")
+	{
+		const FVector2D Icon8x8(8.0f, 8.0f);
+		const FVector2D Icon12x12(12.0f, 12.0f);
+		const FVector2D Icon16x16(16.0f, 16.0f);
+		const FVector2D Icon20x20(20.0f, 20.0f);
+		const FVector2D Icon40x40(40.0f, 40.0f);
+
+		SetContentRoot(FPaths::EnginePluginsDir() / TEXT("Media/MediaPlayerEditor/Content"));
+
+		// buttons
+		Set("MediaPlayerEditor.GoButton", new IMAGE_BRUSH("btn_go_12x", Icon12x12));
+		Set("MediaPlayerEditor.ReloadButton", new IMAGE_BRUSH("btn_reload_12x", Icon12x12));
+		Set("MediaPlayerEditor.SettingsButton", new IMAGE_BRUSH("btn_settings_16x", Icon12x12));
+
+		// misc
+		Set("MediaPlayerEditor.FileMediaSourcePrecached", new IMAGE_BRUSH("filemediasource_precached", Icon12x12));
+		Set("MediaPlayerEditor.MediaSourceOpened", new IMAGE_BRUSH("mediasource_opened", Icon8x8));
+
+		// tabs
+		Set("MediaPlayerEditor.Tabs.Info", new IMAGE_BRUSH("tab_info_16x", Icon16x16));
+		Set("MediaPlayerEditor.Tabs.Media", new IMAGE_BRUSH("tab_media_16x", Icon16x16));
+		Set("MediaPlayerEditor.Tabs.Player", new IMAGE_BRUSH("tab_player_16x", Icon16x16));
+		Set("MediaPlayerEditor.Tabs.Playlist", new IMAGE_BRUSH("tab_playlist_16x", Icon16x16));
+		Set("MediaPlayerEditor.Tabs.Stats", new IMAGE_BRUSH("tab_stats_16x", Icon16x16));
+
+		// toolbar icons
+		Set("MediaPlayerEditor.CloseMedia", new IMAGE_BRUSH("icon_eject_40x", Icon40x40));
+		Set("MediaPlayerEditor.ForwardMedia", new IMAGE_BRUSH("icon_forward_40x", Icon40x40));
+		Set("MediaPlayerEditor.ForwardMedia.Small", new IMAGE_BRUSH("icon_forward_40x", Icon20x20));
+		Set("MediaPlayerEditor.NextMedia", new IMAGE_BRUSH("icon_step_40x", Icon40x40));
+		Set("MediaPlayerEditor.NextMedia.Small", new IMAGE_BRUSH("icon_step_40x", Icon20x20));
+		Set("MediaPlayerEditor.PauseMedia", new IMAGE_BRUSH("icon_pause_40x", Icon40x40));
+		Set("MediaPlayerEditor.PauseMedia.Small", new IMAGE_BRUSH("icon_pause_40x", Icon20x20));
+		Set("MediaPlayerEditor.PlayMedia", new IMAGE_BRUSH("icon_play_40x", Icon40x40));
+		Set("MediaPlayerEditor.PlayMedia.Small", new IMAGE_BRUSH("icon_play_40x", Icon20x20));
+		Set("MediaPlayerEditor.PreviousMedia", new IMAGE_BRUSH("icon_step_back_40x", Icon40x40));
+		Set("MediaPlayerEditor.PreviousMedia.Small", new IMAGE_BRUSH("icon_step_back_40x", Icon20x20));
+		Set("MediaPlayerEditor.ReverseMedia", new IMAGE_BRUSH("icon_reverse_40x", Icon40x40));
+		Set("MediaPlayerEditor.ReverseMedia.Small", new IMAGE_BRUSH("icon_reverse_40x", Icon20x20));
+		Set("MediaPlayerEditor.RewindMedia", new IMAGE_BRUSH("icon_rewind_40x", Icon40x40));
+		Set("MediaPlayerEditor.RewindMedia.Small", new IMAGE_BRUSH("icon_rewind_40x", Icon20x20));
+		Set("MediaPlayerEditor.StopMedia", new IMAGE_BRUSH("icon_stop_40x", Icon40x40));
+		Set("MediaPlayerEditor.StopMedia.Small", new IMAGE_BRUSH("icon_stop_40x", Icon20x20));
+		
+		FSlateStyleRegistry::RegisterSlateStyle(*this);
+	 }
+
+	 /** Destructor. */
+	 ~FMediaPlayerEditorStyle()
+	 {
+		FSlateStyleRegistry::UnRegisterSlateStyle(*this);
+	 }
+};
+
+
+#undef IMAGE_BRUSH
+#undef BOX_BRUSH
+#undef BORDER_BRUSH
+#undef TTF_FONT
+#undef OTF_FONT
diff --git a/Engine/Plugins/Media/MediaPlayerEditor/Source/MediaPlayerEditor/Private/Toolkits/MediaPlayerEditorToolkit.cpp b/Engine/Plugins/Media/MediaPlayerEditor/Source/MediaPlayerEditor/Private/Toolkits/MediaPlayerEditorToolkit.cpp
new file mode 100644
--- /dev/null
+++ b/Engine/Plugins/Media/MediaPlayerEditor/Source/MediaPlayerEditor/Private/Toolkits/MediaPlayerEditorToolkit.cpp
@@ -0,0 +1,433 @@
+// Copyright 1998-2016 Epic Games, Inc. All Rights Reserved.
+
+#include "MediaPlayerEditorPCH.h"
+#include "MediaPlayerEditorCommands.h"
+#include "MediaPlayerEditorToolkit.h"
+#include "SMediaPlayerEditorDetails.h"
+#include "SMediaPlayerEditorInfo.h"
+#include "SMediaPlayerEditorMedia.h"
+#include "SMediaPlayerEditorPlaylist.h"
+#include "SMediaPlayerEditorStats.h"
+#include "SMediaPlayerEditorViewer.h"
+
+
+#define LOCTEXT_NAMESPACE "FMediaPlayerEditorToolkit"
+
+
+namespace MediaPlayerEditorToolkit
+{
+	static const FName AppIdentifier("MediaPlayerEditorApp");
+	static const FName DetailsTabId("Details");
+	static const FName InfoTabId("Info");
+	static const FName MediaTabId("Media");
+	static const FName PlaylistTabId("Playlist");
+	static const FName StatsTabId("Stats");
+	static const FName ViewerTabId("Viewer");
+}
+
+
+/* FMediaPlayerEditorToolkit structors
+ *****************************************************************************/
+
+FMediaPlayerEditorToolkit::FMediaPlayerEditorToolkit(const TSharedRef<ISlateStyle>& InStyle)
+	: MediaPlayer(nullptr)
+	, Style(InStyle)
+{ }
+
+
+FMediaPlayerEditorToolkit::~FMediaPlayerEditorToolkit()
+{
+	FReimportManager::Instance()->OnPreReimport().RemoveAll(this);
+	FReimportManager::Instance()->OnPostReimport().RemoveAll(this);
+
+	GEditor->UnregisterForUndo(this);
+}
+
+
+/* FMediaPlayerEditorToolkit interface
+ *****************************************************************************/
+
+void FMediaPlayerEditorToolkit::Initialize(UMediaPlayer* InMediaPlayer, const EToolkitMode::Type InMode, const TSharedPtr<IToolkitHost>& InToolkitHost)
+{
+	MediaPlayer = InMediaPlayer;
+
+	if (MediaPlayer == nullptr)
+	{
+		return;
+	}
+
+	// support undo/redo
+	MediaPlayer->SetFlags(RF_Transactional);
+	GEditor->RegisterForUndo(this);
+
+	BindCommands();
+
+	// create tab layout
+	const TSharedRef<FTabManager::FLayout> Layout = FTabManager::NewLayout("Standalone_MediaPlayerEditor_v10")
+		->AddArea
+		(
+			FTabManager::NewPrimaryArea()
+				->SetOrientation(Orient_Horizontal)
+				->Split
+				(
+					FTabManager::NewSplitter()
+						->SetOrientation(Orient_Vertical)
+						->SetSizeCoefficient(0.66f)
+						->Split
+						(
+							// tool bar
+							FTabManager::NewStack()
+								->AddTab(GetToolbarTabId(), ETabState::OpenedTab)
+								->SetHideTabWell(true)
+								->SetSizeCoefficient(0.1f)
+								
+						)
+						->Split
+						(
+							// viewer
+							FTabManager::NewStack()
+								->AddTab(MediaPlayerEditorToolkit::ViewerTabId, ETabState::OpenedTab)
+								->SetHideTabWell(true)
+								->SetSizeCoefficient(0.6f)
+						)
+						->Split
+						(
+							// media library
+							FTabManager::NewStack()
+								->AddTab(MediaPlayerEditorToolkit::MediaTabId, ETabState::OpenedTab)
+								->SetHideTabWell(true)
+								->SetSizeCoefficient(0.3f)
+						)
+				)
+				->Split
+				(
+					FTabManager::NewSplitter()
+						->SetOrientation(Orient_Vertical)
+						->SetSizeCoefficient(0.33f)
+						->Split
+						(
+							// playlist
+							FTabManager::NewStack()
+								->AddTab(MediaPlayerEditorToolkit::PlaylistTabId, ETabState::OpenedTab)
+								->SetSizeCoefficient(0.5f)
+						)
+						->Split
+						(
+							// details, info, stats
+							FTabManager::NewStack()
+								->AddTab(MediaPlayerEditorToolkit::DetailsTabId, ETabState::OpenedTab)
+								->AddTab(MediaPlayerEditorToolkit::InfoTabId, ETabState::OpenedTab)
+								->AddTab(MediaPlayerEditorToolkit::StatsTabId, ETabState::ClosedTab)
+								->SetForegroundTab(MediaPlayerEditorToolkit::DetailsTabId)
+								->SetSizeCoefficient(0.5f)
+						)
+				)
+		);
+
+	FAssetEditorToolkit::InitAssetEditor(
+		InMode,
+		InToolkitHost,
+		MediaPlayerEditorToolkit::AppIdentifier,
+		Layout,
+		true /*bCreateDefaultStandaloneMenu*/,
+		true /*bCreateDefaultToolbar*/,
+		InMediaPlayer
+	);
+	
+//	IMediaPlayerEditorModule* MediaPlayerEditorModule = &FModuleManager::LoadModuleChecked<IMediaPlayerEditorModule>("MediaPlayerEditor");
+//	AddMenuExtender(MediaPlayerEditorModule->GetMenuExtensibilityManager()->GetAllExtenders(GetToolkitCommands(), GetEditingObjects()));
+
+	ExtendToolBar();
+	RegenerateMenusAndToolbars();
+}
+
+
+/* FAssetEditorToolkit interface
+ *****************************************************************************/
+
+FString FMediaPlayerEditorToolkit::GetDocumentationLink() const
+{
+	return FString(TEXT("Engine/Content/Types/MediaAssets/Properties/Interface"));
+}
+
+
+bool FMediaPlayerEditorToolkit::OnRequestClose()
+{
+	if (MediaPlayer != nullptr)
+	{
+		MediaPlayer->Close();
+	}
+
+	return FAssetEditorToolkit::OnRequestClose();
+}
+
+
+void FMediaPlayerEditorToolkit::RegisterTabSpawners(const TSharedRef<FTabManager>& InTabManager)
+{
+	WorkspaceMenuCategory = InTabManager->AddLocalWorkspaceMenuCategory(LOCTEXT("WorkspaceMenu_MediaPlayerEditor", "Media Player Editor"));
+	auto WorkspaceMenuCategoryRef = WorkspaceMenuCategory.ToSharedRef();
+
+	FAssetEditorToolkit::RegisterTabSpawners(InTabManager);
+
+	InTabManager->RegisterTabSpawner(MediaPlayerEditorToolkit::DetailsTabId, FOnSpawnTab::CreateSP(this, &FMediaPlayerEditorToolkit::HandleTabManagerSpawnTab, MediaPlayerEditorToolkit::DetailsTabId))
+		.SetDisplayName(LOCTEXT("DetailsTabName", "Details"))
+		.SetGroup(WorkspaceMenuCategoryRef)
+		.SetIcon(FSlateIcon(FEditorStyle::GetStyleSetName(), "LevelEditor.Tabs.Details"));
+
+	InTabManager->RegisterTabSpawner(MediaPlayerEditorToolkit::InfoTabId, FOnSpawnTab::CreateSP(this, &FMediaPlayerEditorToolkit::HandleTabManagerSpawnTab, MediaPlayerEditorToolkit::InfoTabId))
+		.SetDisplayName(LOCTEXT("InfoTabName", "Info"))
+		.SetGroup(WorkspaceMenuCategoryRef)
+		.SetIcon(FSlateIcon(Style->GetStyleSetName(), "MediaPlayerEditor.Tabs.Info"));
+
+	InTabManager->RegisterTabSpawner(MediaPlayerEditorToolkit::MediaTabId, FOnSpawnTab::CreateSP(this, &FMediaPlayerEditorToolkit::HandleTabManagerSpawnTab, MediaPlayerEditorToolkit::MediaTabId))
+		.SetDisplayName(LOCTEXT("MediaTabName", "Media Library"))
+		.SetGroup(WorkspaceMenuCategoryRef)
+		.SetIcon(FSlateIcon(Style->GetStyleSetName(), "MediaPlayerEditor.Tabs.Media"));
+
+	InTabManager->RegisterTabSpawner(MediaPlayerEditorToolkit::ViewerTabId, FOnSpawnTab::CreateSP(this, &FMediaPlayerEditorToolkit::HandleTabManagerSpawnTab, MediaPlayerEditorToolkit::ViewerTabId))
+		.SetDisplayName(LOCTEXT("PlayerTabName", "Player"))
+		.SetGroup(WorkspaceMenuCategoryRef)
+		.SetIcon(FSlateIcon(Style->GetStyleSetName(), "MediaPlayerEditor.Tabs.Player"));
+
+	InTabManager->RegisterTabSpawner(MediaPlayerEditorToolkit::PlaylistTabId, FOnSpawnTab::CreateSP(this, &FMediaPlayerEditorToolkit::HandleTabManagerSpawnTab, MediaPlayerEditorToolkit::PlaylistTabId))
+		.SetDisplayName(LOCTEXT("PlaylistTabName", "Playlist"))
+		.SetGroup(WorkspaceMenuCategoryRef)
+		.SetIcon(FSlateIcon(Style->GetStyleSetName(), "MediaPlayerEditor.Tabs.Playlist"));
+
+	InTabManager->RegisterTabSpawner(MediaPlayerEditorToolkit::StatsTabId, FOnSpawnTab::CreateSP(this, &FMediaPlayerEditorToolkit::HandleTabManagerSpawnTab, MediaPlayerEditorToolkit::StatsTabId))
+		.SetDisplayName(LOCTEXT("StatsTabName", "Stats"))
+		.SetGroup(WorkspaceMenuCategoryRef)
+		.SetIcon(FSlateIcon(Style->GetStyleSetName(), "MediaPlayerEditor.Tabs.Stats"));
+}
+
+
+void FMediaPlayerEditorToolkit::UnregisterTabSpawners(const TSharedRef<class FTabManager>& InTabManager)
+{
+	FAssetEditorToolkit::UnregisterTabSpawners(InTabManager);
+
+	InTabManager->UnregisterTabSpawner(MediaPlayerEditorToolkit::ViewerTabId);
+	InTabManager->UnregisterTabSpawner(MediaPlayerEditorToolkit::StatsTabId);
+	InTabManager->UnregisterTabSpawner(MediaPlayerEditorToolkit::PlaylistTabId);
+	InTabManager->UnregisterTabSpawner(MediaPlayerEditorToolkit::MediaTabId);
+	InTabManager->UnregisterTabSpawner(MediaPlayerEditorToolkit::InfoTabId);
+	InTabManager->UnregisterTabSpawner(MediaPlayerEditorToolkit::DetailsTabId);
+}
+
+
+/* IToolkit interface
+ *****************************************************************************/
+
+FText FMediaPlayerEditorToolkit::GetBaseToolkitName() const
+{
+	return LOCTEXT("AppLabel", "Media Player Editor");
+}
+
+
+FName FMediaPlayerEditorToolkit::GetToolkitFName() const
+{
+	return FName("MediaPlayerEditor");
+}
+
+
+FLinearColor FMediaPlayerEditorToolkit::GetWorldCentricTabColorScale() const
+{
+	return FLinearColor(0.3f, 0.2f, 0.5f, 0.5f);
+}
+
+
+FString FMediaPlayerEditorToolkit::GetWorldCentricTabPrefix() const
+{
+	return LOCTEXT("WorldCentricTabPrefix", "MediaPlayer ").ToString();
+}
+
+
+/* FGCObject interface
+ *****************************************************************************/
+
+void FMediaPlayerEditorToolkit::AddReferencedObjects(FReferenceCollector& Collector)
+{
+	Collector.AddReferencedObject(MediaPlayer);
+}
+
+
+/* FEditorUndoClient interface
+*****************************************************************************/
+
+void FMediaPlayerEditorToolkit::PostUndo(bool bSuccess)
+{
+	// do nothing
+}
+
+
+void FMediaPlayerEditorToolkit::PostRedo(bool bSuccess)
+{
+	PostUndo(bSuccess);
+}
+
+
+/* FMediaPlayerEditorToolkit implementation
+ *****************************************************************************/
+
+void FMediaPlayerEditorToolkit::BindCommands()
+{
+	const FMediaPlayerEditorCommands& Commands = FMediaPlayerEditorCommands::Get();
+
+	ToolkitCommands->MapAction(
+		Commands.CloseMedia,
+		FExecuteAction::CreateLambda([this] { MediaPlayer->Close(); }),
+		FCanExecuteAction::CreateLambda([this] { return !MediaPlayer->GetUrl().IsEmpty(); })
+	);
+
+	ToolkitCommands->MapAction(
+		Commands.ForwardMedia,
+		FExecuteAction::CreateLambda([this]{ MediaPlayer->SetRate(GetForwardRate()); }),
+		FCanExecuteAction::CreateLambda([this]{ return MediaPlayer->IsReady() && MediaPlayer->SupportsRate(GetForwardRate(), false); })
+	);
+
+	ToolkitCommands->MapAction(
+		Commands.NextMedia,
+		FExecuteAction::CreateLambda([this]{ MediaPlayer->Next(); }),
+		FCanExecuteAction::CreateLambda([this]{ return (MediaPlayer->GetPlaylist() != nullptr) && (MediaPlayer->GetPlaylist()->Num() > 1); })
+	);
+
+	ToolkitCommands->MapAction(
+		Commands.PauseMedia,
+		FExecuteAction::CreateLambda([this]{ MediaPlayer->Pause(); }),
+		FCanExecuteAction::CreateLambda([this]{ return MediaPlayer->CanPause() && !MediaPlayer->IsPaused(); })
+	);
+
+	ToolkitCommands->MapAction(
+		Commands.PlayMedia,
+		FExecuteAction::CreateLambda([this]{ MediaPlayer->Play(); }),
+		FCanExecuteAction::CreateLambda([this]{ return MediaPlayer->IsReady() && (MediaPlayer->GetRate() != 1.0f); })
+	);
+
+	ToolkitCommands->MapAction(
+		Commands.PreviousMedia,
+		FExecuteAction::CreateLambda([this]{ MediaPlayer->Previous(); }),
+		FCanExecuteAction::CreateLambda([this]{ return (MediaPlayer->GetPlaylist() != nullptr) && (MediaPlayer->GetPlaylist()->Num() > 1); })
+	);
+
+	ToolkitCommands->MapAction(
+		Commands.ReverseMedia,
+		FExecuteAction::CreateLambda([this]{ MediaPlayer->SetRate(GetReverseRate()); } ),
+		FCanExecuteAction::CreateLambda([this]{ return MediaPlayer->IsReady() && MediaPlayer->SupportsRate(GetReverseRate(), false); })
+	);
+
+	ToolkitCommands->MapAction(
+		Commands.RewindMedia,
+		FExecuteAction::CreateLambda([this]{ MediaPlayer->Rewind(); }),
+		FCanExecuteAction::CreateLambda([this]{ return MediaPlayer->IsReady() && MediaPlayer->SupportsSeeking() && MediaPlayer->GetTime() > FTimespan::Zero(); })
+	);
+}
+
+
+BEGIN_SLATE_FUNCTION_BUILD_OPTIMIZATION
+void FMediaPlayerEditorToolkit::ExtendToolBar()
+{
+	struct Local
+	{
+		static void FillToolbar(FToolBarBuilder& ToolbarBuilder, const TSharedRef<FUICommandList> ToolkitCommands)
+		{
+			ToolbarBuilder.BeginSection("PlaybackControls");
+			{
+				ToolbarBuilder.AddToolBarButton(FMediaPlayerEditorCommands::Get().PreviousMedia);
+				ToolbarBuilder.AddToolBarButton(FMediaPlayerEditorCommands::Get().RewindMedia);
+				ToolbarBuilder.AddToolBarButton(FMediaPlayerEditorCommands::Get().ReverseMedia);
+				ToolbarBuilder.AddToolBarButton(FMediaPlayerEditorCommands::Get().PlayMedia);
+				ToolbarBuilder.AddToolBarButton(FMediaPlayerEditorCommands::Get().PauseMedia);
+				ToolbarBuilder.AddToolBarButton(FMediaPlayerEditorCommands::Get().ForwardMedia);
+				ToolbarBuilder.AddToolBarButton(FMediaPlayerEditorCommands::Get().NextMedia);
+			}
+			ToolbarBuilder.EndSection();
+
+			ToolbarBuilder.BeginSection("MediaControls");
+			{
+				ToolbarBuilder.AddToolBarButton(FMediaPlayerEditorCommands::Get().CloseMedia);
+			}
+			ToolbarBuilder.EndSection();
+		}
+	};
+
+
+	TSharedPtr<FExtender> ToolbarExtender = MakeShareable(new FExtender);
+
+	ToolbarExtender->AddToolBarExtension(
+		"Asset",
+		EExtensionHook::After,
+		GetToolkitCommands(),
+		FToolBarExtensionDelegate::CreateStatic(&Local::FillToolbar, GetToolkitCommands())
+	);
+
+	AddToolbarExtender(ToolbarExtender);
+}
+END_SLATE_FUNCTION_BUILD_OPTIMIZATION
+
+
+float FMediaPlayerEditorToolkit::GetForwardRate() const
+{
+	float Rate = MediaPlayer->GetRate();
+
+	if (Rate < 1.0f)
+	{
+		Rate = 1.0f;
+	}
+
+	return 2.0f * Rate;
+}
+
+
+float FMediaPlayerEditorToolkit::GetReverseRate() const
+{
+	float Rate = MediaPlayer->GetRate();
+
+	if (Rate > -1.0f)
+	{
+		return -1.0f;
+	}
+
+	return 2.0f * Rate;
+}
+
+
+/* FMediaPlayerEditorToolkit callbacks
+ *****************************************************************************/
+
+TSharedRef<SDockTab> FMediaPlayerEditorToolkit::HandleTabManagerSpawnTab(const FSpawnTabArgs& Args, FName TabIdentifier)
+{
+	TSharedPtr<SWidget> TabWidget = SNullWidget::NullWidget;
+
+	if (TabIdentifier == MediaPlayerEditorToolkit::DetailsTabId)
+	{
+		TabWidget = SNew(SMediaPlayerEditorDetails, *MediaPlayer, Style);
+	}
+	else if (TabIdentifier == MediaPlayerEditorToolkit::InfoTabId)
+	{
+		TabWidget = SNew(SMediaPlayerEditorInfo, *MediaPlayer, Style);
+	}
+	else if (TabIdentifier == MediaPlayerEditorToolkit::MediaTabId)
+	{
+		TabWidget = SNew(SMediaPlayerEditorMedia, *MediaPlayer, Style);
+	}
+	else if (TabIdentifier == MediaPlayerEditorToolkit::PlaylistTabId)
+	{
+		TabWidget = SNew(SMediaPlayerEditorPlaylist, *MediaPlayer, Style);
+	}
+	else if (TabIdentifier == MediaPlayerEditorToolkit::StatsTabId)
+	{
+		TabWidget = SNew(SMediaPlayerEditorStats, *MediaPlayer, Style);
+	}
+	else if (TabIdentifier == MediaPlayerEditorToolkit::ViewerTabId)
+	{
+		TabWidget = SNew(SMediaPlayerEditorViewer, *MediaPlayer, Style);
+	}
+
+	return SNew(SDockTab)
+		.TabRole(ETabRole::PanelTab)
+		[
+			TabWidget.ToSharedRef()
+		];
+}
+
+
+#undef LOCTEXT_NAMESPACE
diff --git a/Engine/Plugins/Media/MediaPlayerEditor/Source/MediaPlayerEditor/Private/Toolkits/MediaPlayerEditorToolkit.h b/Engine/Plugins/Media/MediaPlayerEditor/Source/MediaPlayerEditor/Private/Toolkits/MediaPlayerEditorToolkit.h
new file mode 100644
--- /dev/null
+++ b/Engine/Plugins/Media/MediaPlayerEditor/Source/MediaPlayerEditor/Private/Toolkits/MediaPlayerEditorToolkit.h
@@ -0,0 +1,105 @@
+// Copyright 1998-2016 Epic Games, Inc. All Rights Reserved.
+
+#pragma once
+
+#include "Toolkits/AssetEditorToolkit.h"
+#include "EditorUndoClient.h"
+
+
+/**
+ * Implements an Editor toolkit for media players.
+ */
+class FMediaPlayerEditorToolkit
+	: public FAssetEditorToolkit
+	, public FEditorUndoClient
+	, public FGCObject
+{
+public:
+
+	/**
+	 * Creates and initializes a new instance.
+	 *
+	 * @param InStyle The style set to use.
+	 */
+	FMediaPlayerEditorToolkit(const TSharedRef<ISlateStyle>& InStyle);
+
+	/** Virtual destructor. */
+	virtual ~FMediaPlayerEditorToolkit();
+
+public:
+
+	/**
+	 * Initializes the editor tool kit.
+	 *
+	 * @param InMediaPlayer The UMediaPlayer asset to edit.
+	 * @param InMode The mode to create the toolkit in.
+	 * @param InToolkitHost The toolkit host.
+	 */
+	void Initialize(UMediaPlayer* InMediaPlayer, const EToolkitMode::Type InMode, const TSharedPtr<IToolkitHost>& InToolkitHost);
+
+public:
+
+	//~ FAssetEditorToolkit interface
+
+	virtual FString GetDocumentationLink() const override;
+	virtual bool OnRequestClose() override;
+	virtual void RegisterTabSpawners(const TSharedRef<FTabManager>& TabManager) override;
+	virtual void UnregisterTabSpawners(const TSharedRef<FTabManager>& TabManager) override;
+
+public:
+
+	//~ IToolkit interface
+
+	virtual FText GetBaseToolkitName() const override;
+	virtual FName GetToolkitFName() const override;
+	virtual FLinearColor GetWorldCentricTabColorScale() const override;
+	virtual FString GetWorldCentricTabPrefix() const override;
+
+public:
+
+	//~ FGCObject interface
+
+	virtual void AddReferencedObjects(FReferenceCollector& Collector) override;
+	
+protected:
+
+	//~ FEditorUndoClient interface
+
+	virtual void PostUndo(bool bSuccess) override;
+	virtual void PostRedo(bool bSuccess) override;
+
+protected:
+
+	/** Binds the UI commands to delegates. */
+	void BindCommands();
+
+	/** Builds the toolbar widget for the media player editor. */
+	void ExtendToolBar();
+
+	/**
+	 * Gets the playback rate for fast forward.
+	 *
+	 * @return Forward playback rate.
+	 */
+	float GetForwardRate() const;
+
+	/**
+	 * Gets the playback rate for reverse.
+	 *
+	 * @return Reverse playback rate.
+	 */
+	float GetReverseRate() const;
+
+private:
+
+	/** Callback for spawning tabs. */
+	TSharedRef<SDockTab> HandleTabManagerSpawnTab(const FSpawnTabArgs& Args, FName TabIdentifier);
+
+private:
+
+	/** The media player asset being edited. */
+	UMediaPlayer* MediaPlayer;
+
+	/** Pointer to the style set to use for toolkits. */
+	TSharedRef<ISlateStyle> Style;
+};
