Copying //UE4/Dev-Sequencer to //UE4/Dev-Main (Source: //UE4/Dev-Sequencer @ 3057646)

diff --git a/Engine/Plugins/Media/MediaPlayerEditor/Source/MediaPlayerEditor/Private/Toolkits/MediaPlaylistEditorToolkit.cpp b/Engine/Plugins/Media/MediaPlayerEditor/Source/MediaPlayerEditor/Private/Toolkits/MediaPlaylistEditorToolkit.cpp
new file mode 100644
--- /dev/null
+++ b/Engine/Plugins/Media/MediaPlayerEditor/Source/MediaPlayerEditor/Private/Toolkits/MediaPlaylistEditorToolkit.cpp
@@ -0,0 +1,211 @@
+// Copyright 1998-2016 Epic Games, Inc. All Rights Reserved.
+
+#include "MediaPlayerEditorPCH.h"
+#include "MediaPlaylistEditorToolkit.h"
+#include "SMediaPlaylistEditorDetails.h"
+#include "SMediaPlaylistEditorMedia.h"
+
+
+#define LOCTEXT_NAMESPACE "FMediaPlaylistEditorToolkit"
+
+
+namespace MediaPlaylistEditorToolkit
+{
+	static const FName AppIdentifier("MediaPlaylistEditorApp");
+	static const FName DetailsTabId("Details");
+	static const FName MediaTabId("Media");
+}
+
+
+/* FMediaPlaylistEditorToolkit structors
+ *****************************************************************************/
+
+FMediaPlaylistEditorToolkit::FMediaPlaylistEditorToolkit(const TSharedRef<ISlateStyle>& InStyle)
+	: MediaPlaylist(nullptr)
+	, Style(InStyle)
+{ }
+
+
+FMediaPlaylistEditorToolkit::~FMediaPlaylistEditorToolkit()
+{
+	FReimportManager::Instance()->OnPreReimport().RemoveAll(this);
+	FReimportManager::Instance()->OnPostReimport().RemoveAll(this);
+
+	GEditor->UnregisterForUndo(this);
+}
+
+
+/* FMediaPlaylistEditorToolkit interface
+ *****************************************************************************/
+
+void FMediaPlaylistEditorToolkit::Initialize(UMediaPlaylist* InMediaPlaylist, const EToolkitMode::Type InMode, const TSharedPtr<IToolkitHost>& InToolkitHost)
+{
+	MediaPlaylist = InMediaPlaylist;
+
+	if (MediaPlaylist == nullptr)
+	{
+		return;
+	}
+
+	// Support undo/redo
+	MediaPlaylist->SetFlags(RF_Transactional);
+	GEditor->RegisterForUndo(this);
+
+	// create tab layout
+	const TSharedRef<FTabManager::FLayout> Layout = FTabManager::NewLayout("Standalone_MediaPlaylistEditor_v2")
+		->AddArea
+		(
+			FTabManager::NewPrimaryArea()
+				->SetOrientation(Orient_Vertical)
+				->Split
+				(
+					// tool bar
+					FTabManager::NewStack()
+						->AddTab(GetToolbarTabId(), ETabState::OpenedTab)
+						->SetHideTabWell(true)
+						->SetSizeCoefficient(0.1f)
+								
+				)
+				->Split
+				(
+					// details
+					FTabManager::NewStack()
+						->AddTab(MediaPlaylistEditorToolkit::DetailsTabId, ETabState::OpenedTab)
+						->SetSizeCoefficient(0.45f)
+				)
+				->Split
+				(
+					// media library
+					FTabManager::NewStack()
+						->AddTab(MediaPlaylistEditorToolkit::MediaTabId, ETabState::OpenedTab)
+						->SetHideTabWell(true)
+						->SetSizeCoefficient(0.45f)
+				)
+		);
+
+	FAssetEditorToolkit::InitAssetEditor(
+		InMode,
+		InToolkitHost,
+		MediaPlaylistEditorToolkit::AppIdentifier,
+		Layout,
+		true /*bCreateDefaultStandaloneMenu*/,
+		true /*bCreateDefaultToolbar*/,
+		InMediaPlaylist
+	);
+	
+	RegenerateMenusAndToolbars();
+}
+
+
+/* FAssetEditorToolkit interface
+ *****************************************************************************/
+
+FString FMediaPlaylistEditorToolkit::GetDocumentationLink() const
+{
+	return FString(TEXT("Engine/Content/Types/MediaAssets/Properties/Interface"));
+}
+
+
+void FMediaPlaylistEditorToolkit::RegisterTabSpawners(const TSharedRef<FTabManager>& InTabManager)
+{
+	WorkspaceMenuCategory = InTabManager->AddLocalWorkspaceMenuCategory(LOCTEXT("WorkspaceMenu_MediaPlaylistEditor", "Media Player Editor"));
+	auto WorkspaceMenuCategoryRef = WorkspaceMenuCategory.ToSharedRef();
+
+	FAssetEditorToolkit::RegisterTabSpawners(InTabManager);
+
+	InTabManager->RegisterTabSpawner(MediaPlaylistEditorToolkit::DetailsTabId, FOnSpawnTab::CreateSP(this, &FMediaPlaylistEditorToolkit::HandleTabManagerSpawnTab, MediaPlaylistEditorToolkit::DetailsTabId))
+		.SetDisplayName(LOCTEXT("DetailsTabName", "Details"))
+		.SetGroup(WorkspaceMenuCategoryRef)
+		.SetIcon(FSlateIcon(FEditorStyle::GetStyleSetName(), "LevelEditor.Tabs.Details"));
+
+	InTabManager->RegisterTabSpawner(MediaPlaylistEditorToolkit::MediaTabId, FOnSpawnTab::CreateSP(this, &FMediaPlaylistEditorToolkit::HandleTabManagerSpawnTab, MediaPlaylistEditorToolkit::MediaTabId))
+		.SetDisplayName(LOCTEXT("MediaTabName", "Media Library"))
+		.SetGroup(WorkspaceMenuCategoryRef)
+		.SetIcon(FSlateIcon(Style->GetStyleSetName(), "MediaPlayerEditor.Tabs.Media"));
+}
+
+
+void FMediaPlaylistEditorToolkit::UnregisterTabSpawners(const TSharedRef<FTabManager>& InTabManager)
+{
+	FAssetEditorToolkit::UnregisterTabSpawners(InTabManager);
+
+	InTabManager->UnregisterTabSpawner(MediaPlaylistEditorToolkit::MediaTabId);	
+	InTabManager->UnregisterTabSpawner(MediaPlaylistEditorToolkit::DetailsTabId);
+}
+
+
+/* IToolkit interface
+ *****************************************************************************/
+
+FText FMediaPlaylistEditorToolkit::GetBaseToolkitName() const
+{
+	return LOCTEXT("AppLabel", "Media Playlist Editor");
+}
+
+
+FName FMediaPlaylistEditorToolkit::GetToolkitFName() const
+{
+	return FName("MediaPlaylistEditor");
+}
+
+
+FLinearColor FMediaPlaylistEditorToolkit::GetWorldCentricTabColorScale() const
+{
+	return FLinearColor(0.3f, 0.2f, 0.5f, 0.5f);
+}
+
+
+FString FMediaPlaylistEditorToolkit::GetWorldCentricTabPrefix() const
+{
+	return LOCTEXT("WorldCentricTabPrefix", "MediaPlaylist ").ToString();
+}
+
+
+/* FGCObject interface
+ *****************************************************************************/
+
+void FMediaPlaylistEditorToolkit::AddReferencedObjects(FReferenceCollector& Collector)
+{
+	Collector.AddReferencedObject(MediaPlaylist);
+}
+
+
+/* FEditorUndoClient interface
+*****************************************************************************/
+
+void FMediaPlaylistEditorToolkit::PostUndo(bool bSuccess)
+{
+}
+
+
+void FMediaPlaylistEditorToolkit::PostRedo(bool bSuccess)
+{
+	PostUndo(bSuccess);
+}
+
+
+/* FMediaPlaylistEditorToolkit callbacks
+ *****************************************************************************/
+
+TSharedRef<SDockTab> FMediaPlaylistEditorToolkit::HandleTabManagerSpawnTab(const FSpawnTabArgs& Args, FName TabIdentifier)
+{
+	TSharedPtr<SWidget> TabWidget = SNullWidget::NullWidget;
+
+	if (TabIdentifier == MediaPlaylistEditorToolkit::DetailsTabId)
+	{
+		TabWidget = SNew(SMediaPlaylistEditorDetails, *MediaPlaylist, Style);
+	}
+	else if (TabIdentifier == MediaPlaylistEditorToolkit::MediaTabId)
+	{
+		TabWidget = SNew(SMediaPlaylistEditorMedia, *MediaPlaylist, Style);
+	}
+
+	return SNew(SDockTab)
+		.TabRole(ETabRole::PanelTab)
+		[
+			TabWidget.ToSharedRef()
+		];
+}
+
+
+#undef LOCTEXT_NAMESPACE
diff --git a/Engine/Plugins/Media/MediaPlayerEditor/Source/MediaPlayerEditor/Private/Toolkits/MediaPlaylistEditorToolkit.h b/Engine/Plugins/Media/MediaPlayerEditor/Source/MediaPlayerEditor/Private/Toolkits/MediaPlaylistEditorToolkit.h
new file mode 100644
--- /dev/null
+++ b/Engine/Plugins/Media/MediaPlayerEditor/Source/MediaPlayerEditor/Private/Toolkits/MediaPlaylistEditorToolkit.h
@@ -0,0 +1,82 @@
+// Copyright 1998-2016 Epic Games, Inc. All Rights Reserved.
+
+#pragma once
+
+#include "Toolkits/AssetEditorToolkit.h"
+#include "EditorUndoClient.h"
+
+
+/**
+ * Implements an Editor toolkit for media play lists.
+ */
+class FMediaPlaylistEditorToolkit
+	: public FAssetEditorToolkit
+	, public FEditorUndoClient
+	, public FGCObject
+{
+public:
+
+	/**
+	 * Creates and initializes a new instance.
+	 *
+	 * @param InStyle The style set to use.
+	 */
+	FMediaPlaylistEditorToolkit(const TSharedRef<ISlateStyle>& InStyle);
+
+	/** Virtual destructor. */
+	virtual ~FMediaPlaylistEditorToolkit();
+
+public:
+
+	/**
+	 * Initializes the editor tool kit.
+	 *
+	 * @param InMediaPlaylist The UMediaPlaylist asset to edit.
+	 * @param InMode The mode to create the toolkit in.
+	 * @param InToolkitHost The toolkit host.
+	 */
+	void Initialize(UMediaPlaylist* InMediaPlaylist, const EToolkitMode::Type InMode, const TSharedPtr<IToolkitHost>& InToolkitHost);
+
+public:
+
+	//~ FAssetEditorToolkit interface
+
+	virtual FString GetDocumentationLink() const override;
+	virtual void RegisterTabSpawners(const TSharedRef<FTabManager>& TabManager) override;
+	virtual void UnregisterTabSpawners(const TSharedRef<FTabManager>& TabManager) override;
+
+public:
+
+	//~ IToolkit interface
+
+	virtual FText GetBaseToolkitName() const override;
+	virtual FName GetToolkitFName() const override;
+	virtual FLinearColor GetWorldCentricTabColorScale() const override;
+	virtual FString GetWorldCentricTabPrefix() const override;
+
+public:
+
+	//~ FGCObject interface
+
+	virtual void AddReferencedObjects(FReferenceCollector& Collector) override;
+	
+protected:
+
+	//~ FEditorUndoClient interface
+
+	virtual void PostUndo(bool bSuccess) override;
+	virtual void PostRedo(bool bSuccess) override;
+
+private:
+
+	/** Callback for spawning tabs. */
+	TSharedRef<SDockTab> HandleTabManagerSpawnTab(const FSpawnTabArgs& Args, FName TabIdentifier);
+
+private:
+
+	/** The media play list asset being edited. */
+	UMediaPlaylist* MediaPlaylist;
+
+	/** Pointer to the style set to use for toolkits. */
+	TSharedRef<ISlateStyle> Style;
+};
diff --git a/Engine/Plugins/Media/MediaPlayerEditor/Source/MediaPlayerEditor/Private/Widgets/SMediaPlayerEditorDetails.cpp b/Engine/Plugins/Media/MediaPlayerEditor/Source/MediaPlayerEditor/Private/Widgets/SMediaPlayerEditorDetails.cpp
--- a/Engine/Plugins/Media/MediaPlayerEditor/Source/MediaPlayerEditor/Private/Widgets/SMediaPlayerEditorDetails.cpp
+++ b/Engine/Plugins/Media/MediaPlayerEditor/Source/MediaPlayerEditor/Private/Widgets/SMediaPlayerEditorDetails.cpp
@@ -1,6 +1,7 @@
 // Copyright 1998-2016 Epic Games, Inc. All Rights Reserved.
 
-#include "MediaPlayerEditorPrivatePCH.h"
+#include "MediaPlayerEditorPCH.h"
+#include "SMediaPlayerEditorDetails.h"
 
 
 #define LOCTEXT_NAMESPACE "SMediaPlayerEditorDetails"
@@ -9,9 +10,9 @@
 /* SMediaPlayerEditorDetails interface
  *****************************************************************************/
 
-void SMediaPlayerEditorDetails::Construct( const FArguments& InArgs, UMediaPlayer* InMediaPlayer, const TSharedRef<ISlateStyle>& InStyle )
+void SMediaPlayerEditorDetails::Construct(const FArguments& InArgs, UMediaPlayer& InMediaPlayer, const TSharedRef<ISlateStyle>& InStyle)
 {
-	MediaPlayer = InMediaPlayer;
+	MediaPlayer = &InMediaPlayer;
 
 	// initialize details view
 	FDetailsViewArgs DetailsViewArgs;
@@ -30,14 +31,7 @@ void SMediaPlayerEditorDetails::Construct( const FArguments& InArgs, UMediaPlaye
 
 	ChildSlot
 	[
-		SNew(SVerticalBox)
-
-		+ SVerticalBox::Slot()
-			.AutoHeight()
-			.Padding(4.0f, 2.0f)
-			[
-				DetailsView.ToSharedRef()
-			]
+		DetailsView.ToSharedRef()
 	];
 }
 
diff --git a/Engine/Plugins/Media/MediaPlayerEditor/Source/MediaPlayerEditor/Private/Widgets/SMediaPlayerEditorDetails.h b/Engine/Plugins/Media/MediaPlayerEditor/Source/MediaPlayerEditor/Private/Widgets/SMediaPlayerEditorDetails.h
--- a/Engine/Plugins/Media/MediaPlayerEditor/Source/MediaPlayerEditor/Private/Widgets/SMediaPlayerEditorDetails.h
+++ b/Engine/Plugins/Media/MediaPlayerEditor/Source/MediaPlayerEditor/Private/Widgets/SMediaPlayerEditorDetails.h
@@ -23,7 +23,7 @@ public:
 	 * @param InMediaPlayer The MediaPlayer asset to show the details for.
 	 * @param InStyleSet The style set to use.
 	 */
-	void Construct( const FArguments& InArgs, UMediaPlayer* InMediaPlayer, const TSharedRef<ISlateStyle>& InStyle );
+	void Construct(const FArguments& InArgs, UMediaPlayer& InMediaPlayer, const TSharedRef<ISlateStyle>& InStyle);
 
 private:
 
diff --git a/Engine/Plugins/Media/MediaPlayerEditor/Source/MediaPlayerEditor/Private/Widgets/SMediaPlayerEditorInfo.cpp b/Engine/Plugins/Media/MediaPlayerEditor/Source/MediaPlayerEditor/Private/Widgets/SMediaPlayerEditorInfo.cpp
new file mode 100644
--- /dev/null
+++ b/Engine/Plugins/Media/MediaPlayerEditor/Source/MediaPlayerEditor/Private/Widgets/SMediaPlayerEditorInfo.cpp
@@ -0,0 +1,116 @@
+// Copyright 1998-2016 Epic Games, Inc. All Rights Reserved.
+
+#include "MediaPlayerEditorPCH.h"
+#include "SMediaPlayerEditorInfo.h"
+
+
+#define LOCTEXT_NAMESPACE "SMediaPlayerEditorInfo"
+
+
+/* SMediaPlayerEditorInfo interface
+ *****************************************************************************/
+
+void SMediaPlayerEditorInfo::Construct(const FArguments& InArgs, UMediaPlayer& InMediaPlayer, const TSharedRef<ISlateStyle>& InStyle)
+{
+	MediaPlayer = &InMediaPlayer;
+
+	ChildSlot
+	[
+		SNew(SBorder)
+			.BorderImage(FEditorStyle::GetBrush("ToolPanel.GroupBorder"))
+			.Padding(4.0)
+			[
+				SNew(SVerticalBox)
+
+				+ SVerticalBox::Slot()
+					.FillHeight(1.0f)
+					[
+						SAssignNew(MediaInfoBox, SMultiLineEditableTextBox)
+							.HintText(LOCTEXT("NoMediaOpened", "No media opened"))
+							.IsReadOnly(true)
+					]
+			]
+	];
+
+	MediaPlayer->OnMediaEvent().AddSP(this, &SMediaPlayerEditorInfo::HandleMediaPlayerMediaEvent);
+}
+
+
+/* SMediaPlayerEditorInfo callbacks
+ *****************************************************************************/
+
+void SMediaPlayerEditorInfo::HandleMediaPlayerMediaEvent(EMediaEvent Event)
+{
+	if ((Event == EMediaEvent::MediaOpened) || (Event == EMediaEvent::MediaOpenFailed) || (Event == EMediaEvent::TracksChanged))
+	{
+		TSharedPtr<IMediaPlayer> Player = MediaPlayer->GetPlayer();
+
+		if (Player.IsValid())
+		{
+			TRange<float> ThinnedForwardRates = Player->GetControls().GetSupportedRates(EMediaPlaybackDirections::Forward, false);
+			TRange<float> UnthinnedForwardRates = Player->GetControls().GetSupportedRates(EMediaPlaybackDirections::Forward, true);
+			TRange<float> ThinnedReverseRates = Player->GetControls().GetSupportedRates(EMediaPlaybackDirections::Reverse, false);
+			TRange<float> UnthinnedReverseRates = Player->GetControls().GetSupportedRates(EMediaPlaybackDirections::Reverse, true);
+
+			FFormatNamedArguments Arguments;
+			{
+				Arguments.Add(TEXT("PlayerName"), FText::FromName(MediaPlayer->GetPlayerName()));
+				Arguments.Add(TEXT("SupportsScrubbing"), MediaPlayer->SupportsScrubbing() ? GYes : GNo);
+				Arguments.Add(TEXT("SupportsSeeking"), MediaPlayer->SupportsSeeking() ? GYes : GNo);
+				Arguments.Add(TEXT("PlayerInfo"), FText::FromString(Player->GetInfo()));
+
+				Arguments.Add(TEXT("ForwardRateThinned"), ThinnedForwardRates.IsEmpty()
+					? LOCTEXT("RateNotSupported", "Not supported")
+					: FText::Format(LOCTEXT("RatesFormat", "{0} to {1}"), FText::AsNumber(ThinnedForwardRates.GetLowerBoundValue()), FText::AsNumber(ThinnedForwardRates.GetUpperBoundValue()))
+				);
+
+				Arguments.Add(TEXT("ForwardRateUnthinned"), UnthinnedForwardRates.IsEmpty()
+					? LOCTEXT("RateNotSupported", "Not supported")
+					: FText::Format(LOCTEXT("RatesFormat", "{0} to {1}"), FText::AsNumber(UnthinnedForwardRates.GetLowerBoundValue()), FText::AsNumber(UnthinnedForwardRates.GetUpperBoundValue()))
+				);
+
+				Arguments.Add(TEXT("ReverseRateThinned"), ThinnedReverseRates.IsEmpty()
+					? LOCTEXT("RateNotSupported", "Not supported")
+					: FText::Format(LOCTEXT("RatesFormat", "{0} to {1}"), FText::AsNumber(ThinnedReverseRates.GetLowerBoundValue()), FText::AsNumber(ThinnedReverseRates.GetUpperBoundValue()))
+				);
+
+				Arguments.Add(TEXT("ReverseRateUnthinned"), UnthinnedReverseRates.IsEmpty()
+					? LOCTEXT("RateNotSupported", "Not supported")
+					: FText::Format(LOCTEXT("RatesFormat", "{0} to {1}"), FText::AsNumber(UnthinnedReverseRates.GetLowerBoundValue()), FText::AsNumber(UnthinnedReverseRates.GetUpperBoundValue()))
+				);				
+			}
+
+			MediaInfoBox->SetText(
+				FText::Format(LOCTEXT("InfoFormat",
+					"Player: {PlayerName}\n"
+					"\n"
+					"Forward Rates\n"
+					"    Thinned: {ForwardRateThinned}\n"
+					"    Unthinned: {ForwardRateUnthinned}\n"
+					"\n"
+					"Reverse Rates\n"
+					"    Thinned: {ReverseRateThinned}\n"
+					"    Unthinned: {ReverseRateUnthinned}\n"
+					"\n"
+					"Capabilities\n"
+					"    Scrubbing: {SupportsScrubbing}\n"
+					"    Seeking: {SupportsSeeking}\n"
+					"\n"
+					"{PlayerInfo}"),
+					Arguments
+				)
+			);
+		}
+		else
+		{
+			MediaInfoBox->SetText(FText::GetEmpty());
+		}
+	}
+	else if (Event == EMediaEvent::MediaClosed)
+	{
+		MediaInfoBox->SetText(FText::GetEmpty());
+	}
+}
+
+
+#undef LOCTEXT_NAMESPACE
diff --git a/Engine/Plugins/Media/MediaPlayerEditor/Source/MediaPlayerEditor/Private/Widgets/SMediaPlayerEditorMedia.cpp b/Engine/Plugins/Media/MediaPlayerEditor/Source/MediaPlayerEditor/Private/Widgets/SMediaPlayerEditorMedia.cpp
new file mode 100644
--- /dev/null
+++ b/Engine/Plugins/Media/MediaPlayerEditor/Source/MediaPlayerEditor/Private/Widgets/SMediaPlayerEditorMedia.cpp
@@ -0,0 +1,199 @@
+// Copyright 1998-2016 Epic Games, Inc. All Rights Reserved.
+
+#include "MediaPlayerEditorPCH.h"
+#include "SMediaPlayerEditorMedia.h"
+
+
+#define LOCTEXT_NAMESPACE "SMediaPlayerEditorMedia"
+
+
+/* SMediaPlayerEditorMedia interface
+ *****************************************************************************/
+
+void SMediaPlayerEditorMedia::Construct(const FArguments& InArgs, UMediaPlayer& InMediaPlayer, const TSharedRef<ISlateStyle>& InStyle)
+{
+	MediaPlayer = &InMediaPlayer;
+	Style = InStyle;
+
+	// initialize asset picker
+	FAssetPickerConfig AssetPickerConfig;
+	{
+		AssetPickerConfig.Filter.ClassNames.Add(UMediaPlaylist::StaticClass()->GetFName());
+		AssetPickerConfig.Filter.ClassNames.Add(UMediaSource::StaticClass()->GetFName());
+		AssetPickerConfig.Filter.bRecursiveClasses = true;
+
+		AssetPickerConfig.AssetShowWarningText = LOCTEXT("NoMediaSourcesFound", "No media sources or play lists found.");
+		AssetPickerConfig.bAllowDragging = false;
+		AssetPickerConfig.bAutohideSearchBar = true;
+		AssetPickerConfig.bCanShowClasses = false;
+		AssetPickerConfig.bCanShowDevelopersFolder = true;
+		AssetPickerConfig.InitialAssetViewType = EAssetViewType::Column;
+		AssetPickerConfig.ThumbnailScale = 0.1f;
+
+		AssetPickerConfig.OnAssetDoubleClicked = FOnAssetDoubleClicked::CreateSP(this, &SMediaPlayerEditorMedia::HandleAssetPickerAssetDoubleClicked);
+		AssetPickerConfig.OnAssetEnterPressed = FOnAssetEnterPressed::CreateSP(this, &SMediaPlayerEditorMedia::HandleAssetPickerAssetEnterPressed);
+		AssetPickerConfig.OnGetAssetContextMenu = FOnGetAssetContextMenu::CreateSP(this, &SMediaPlayerEditorMedia::HandleAssetPickerGetAssetContextMenu);
+	}
+
+	auto& ContentBrowserModule = FModuleManager::Get().LoadModuleChecked<FContentBrowserModule>(TEXT("ContentBrowser"));
+
+	ChildSlot
+	[
+		SNew(SBorder)
+			.BorderImage(FEditorStyle::GetBrush("ToolPanel.GroupBorder"))
+			.ToolTipText(LOCTEXT("DoubleClickToAddToolTip", "Double-click a media source or playlist to open it in the player."))
+			[
+				ContentBrowserModule.Get().CreateAssetPicker(AssetPickerConfig)
+			]
+	];
+}
+
+
+/* SMediaPlayerEditorMedia implementation
+ *****************************************************************************/
+
+void SMediaPlayerEditorMedia::OpenMediaAsset(UObject* Asset)
+{
+	UMediaSource* MediaSource = Cast<UMediaSource>(Asset);
+
+	if (MediaSource != nullptr)
+	{
+		if (!MediaPlayer->OpenSource(MediaSource))
+		{
+			ShowMediaOpenFailedMessage();
+		}
+
+		return;
+	}
+
+	UMediaPlaylist* MediaPlaylist = Cast<UMediaPlaylist>(Asset);
+
+	if (MediaPlaylist != nullptr)
+	{
+		if (!MediaPlayer->OpenPlaylist(MediaPlaylist))
+		{
+			ShowMediaOpenFailedMessage();
+		}
+	}
+}
+
+
+void SMediaPlayerEditorMedia::ShowMediaOpenFailedMessage()
+{
+	FNotificationInfo NotificationInfo(LOCTEXT("MediaOpenFailedError", "The media failed to open. Check Output Log for details!"));
+	{
+		NotificationInfo.ExpireDuration = 2.0f;
+	}
+
+	FSlateNotificationManager::Get().AddNotification(NotificationInfo)->SetCompletionState(SNotificationItem::CS_Fail);
+}
+
+
+/* SMediaPlayerEditorMedia callbacks
+ *****************************************************************************/
+
+void SMediaPlayerEditorMedia::HandleAssetPickerAssetDoubleClicked(const class FAssetData& AssetData)
+{
+	OpenMediaAsset(AssetData.GetAsset());
+}
+
+
+void SMediaPlayerEditorMedia::HandleAssetPickerAssetEnterPressed(const TArray<FAssetData>& SelectedAssets)
+{
+	if (SelectedAssets.Num() > 0)
+	{
+		OpenMediaAsset(SelectedAssets[0].GetAsset());
+	}
+}
+
+
+TSharedPtr<SWidget> SMediaPlayerEditorMedia::HandleAssetPickerGetAssetContextMenu(const TArray<FAssetData>& SelectedAssets)
+{
+	if (SelectedAssets.Num() <= 0)
+	{
+		return nullptr;
+	}
+
+	UObject* SelectedAsset = SelectedAssets[0].GetAsset();
+
+	if (SelectedAsset == nullptr)
+	{
+		return nullptr;
+	}
+
+	FMenuBuilder MenuBuilder(true /*bInShouldCloseWindowAfterMenuSelection=*/, nullptr);
+	{
+		// media section
+		MenuBuilder.BeginSection("MediaSection", LOCTEXT("MediaSection", "Media"));
+		{
+			MenuBuilder.AddMenuEntry(
+				LOCTEXT("EditMenuAction", "Edit..."),
+				LOCTEXT("EditMenuActionTooltip", "Opens the selected asset for edit."),
+				FSlateIcon(FEditorStyle::GetStyleSetName(), "ContentBrowser.AssetActions.Edit"),
+				FUIAction(
+					FExecuteAction::CreateLambda([=]() {
+						FAssetEditorManager::Get().OpenEditorForAsset(SelectedAsset);
+					})
+				)
+			);
+
+			MenuBuilder.AddMenuEntry(
+				LOCTEXT("LoadMenuAction", "Open"),
+				LOCTEXT("LoadMenuActionTooltip", "Open this media asset in the player"),
+				FSlateIcon(Style->GetStyleSetName(), "MediaPlayerEditor.NextMedia.Small"),
+				FUIAction(
+					FExecuteAction::CreateLambda([=]() {
+						OpenMediaAsset(SelectedAsset);
+					})
+				)
+			);
+		}
+		MenuBuilder.EndSection();
+
+		// asset section
+		MenuBuilder.BeginSection("AssetSection", LOCTEXT("AssetSection", "Asset"));
+		{
+			MenuBuilder.AddMenuEntry(
+				LOCTEXT("FindInCbMenuAction", "Find in Content Browser"),
+				LOCTEXT("FindInCbMenuActionTooltip", "Summons the Content Browser and navigates to the selected asset"),
+				FSlateIcon(FEditorStyle::GetStyleSetName(), "SystemWideCommands.FindInContentBrowser"),
+				FUIAction(
+					FExecuteAction::CreateLambda([=]() {
+						TArray<UObject*> AssetsToSync = TArrayBuilder<UObject*>().Add(SelectedAsset);
+						GEditor->SyncBrowserToObjects(AssetsToSync);
+					})
+				)
+			);
+
+			auto FileMediaSource = Cast<UFileMediaSource>(SelectedAsset);
+
+			if (FileMediaSource != nullptr)
+			{
+				FFormatNamedArguments Args;
+				{
+					Args.Add(TEXT("FileManagerName"), FPlatformMisc::GetFileManagerName());
+				}
+
+				MenuBuilder.AddMenuEntry(
+					FText::Format(LOCTEXT("OpenInFileManager", "Show Media File in {FileManagerName}"), Args),
+					LOCTEXT("OpenInFileManagerTooltip", "Finds the media file that this asset points to on disk"),
+					FSlateIcon(FEditorStyle::GetStyleSetName(), "SystemWideCommands.FindInContentBrowser"),
+					FUIAction(
+						FExecuteAction::CreateLambda([=]() {
+							FPlatformProcess::ExploreFolder(*FileMediaSource->FilePath);
+						}),
+						FCanExecuteAction::CreateLambda([=]() -> bool {
+							return FileMediaSource->Validate();
+						})
+					)
+				);
+			}
+		}
+		MenuBuilder.EndSection();
+	}
+
+	return MenuBuilder.MakeWidget();
+}
+
+
+#undef LOCTEXT_NAMESPACE
diff --git a/Engine/Plugins/Media/MediaPlayerEditor/Source/MediaPlayerEditor/Private/Widgets/SMediaPlayerEditorOutput.cpp b/Engine/Plugins/Media/MediaPlayerEditor/Source/MediaPlayerEditor/Private/Widgets/SMediaPlayerEditorOutput.cpp
new file mode 100644
--- /dev/null
+++ b/Engine/Plugins/Media/MediaPlayerEditor/Source/MediaPlayerEditor/Private/Widgets/SMediaPlayerEditorOutput.cpp
@@ -0,0 +1,250 @@
+// Copyright 1998-2016 Epic Games, Inc. All Rights Reserved.
+
+#include "MediaPlayerEditorPCH.h"
+#include "SMediaPlayerEditorOutput.h"
+
+
+/* SMediaPlayerEditorOutput structors
+ *****************************************************************************/
+
+SMediaPlayerEditorOutput::SMediaPlayerEditorOutput()
+	: CurrentSoundWave(nullptr)
+	, CurrentTexture(nullptr)
+	, DefaultSoundWave(nullptr)
+	, DefaultTexture(nullptr)
+	, Material(nullptr)
+	, MediaPlayer(nullptr)
+{
+	DefaultSoundWave = NewObject<UMediaSoundWave>(GetTransientPackage(), NAME_None, RF_Transient | RF_Public);
+
+	if (DefaultSoundWave != nullptr)
+	{
+		DefaultSoundWave->AddToRoot();
+
+		AudioComponent = FAudioDevice::CreateComponent(DefaultSoundWave, nullptr, nullptr, false, false);
+	
+		if (AudioComponent != nullptr)
+		{
+			AudioComponent->SetVolumeMultiplier(1.0f);
+			AudioComponent->SetPitchMultiplier(1.0f);
+			AudioComponent->bAllowSpatialization = false;
+			AudioComponent->bIsUISound = true;
+			AudioComponent->bAutoDestroy = false;
+			AudioComponent->AddToRoot();
+		}
+	}
+}
+
+
+SMediaPlayerEditorOutput::~SMediaPlayerEditorOutput()
+{
+	if (MediaPlayer != nullptr)
+	{
+		MediaPlayer->OnMediaEvent().RemoveAll(this);
+	}
+
+	if (AudioComponent != nullptr)
+	{
+		AudioComponent->Stop();
+		AudioComponent->RemoveFromRoot();
+	}
+
+	if (Material != nullptr)
+	{
+		Material->RemoveFromRoot();
+	}
+
+	if (DefaultSoundWave != nullptr)
+	{
+		DefaultSoundWave->RemoveFromRoot();
+	}
+
+	if (DefaultTexture != nullptr)
+	{
+		DefaultTexture->RemoveFromRoot();
+	}
+}
+
+
+/* SMediaPlayerEditorOutput interface
+ *****************************************************************************/
+
+void SMediaPlayerEditorOutput::Construct(const FArguments& InArgs, UMediaPlayer& InMediaPlayer)
+{
+	MediaPlayer = &InMediaPlayer;
+
+	// create wrapper material
+	Material = NewObject<UMaterial>(GetTransientPackage(), NAME_None, RF_Transient | RF_Public);
+
+	if (Material != nullptr)
+	{
+		TextureSampler = NewObject<UMaterialExpressionTextureSample>(Material);
+		{
+			TextureSampler->SamplerType = SAMPLERTYPE_Color;
+		}
+
+		FExpressionOutput& Output = TextureSampler->GetOutputs()[0];
+		FExpressionInput& Input = Material->EmissiveColor;
+		{
+			Input.Expression = TextureSampler;
+			Input.Mask = Output.Mask;
+			Input.MaskR = Output.MaskR;
+			Input.MaskG = Output.MaskG;
+			Input.MaskB = Output.MaskB;
+			Input.MaskA = Output.MaskA;
+		}
+
+		Material->Expressions.Add(TextureSampler);
+		Material->MaterialDomain = EMaterialDomain::MD_UI;
+		Material->AddToRoot();
+	}
+
+	// create Slate brush
+	MaterialBrush = MakeShareable(new FSlateBrush());
+	{
+		MaterialBrush->SetResourceObject(Material);
+	}
+
+	UpdateMaterial();
+
+	ChildSlot
+	[
+		SNew(SImage)
+			.Image(MaterialBrush.Get())
+	];
+
+	MediaPlayer->OnMediaEvent().AddRaw(this, &SMediaPlayerEditorOutput::HandleMediaPlayerMediaEvent);
+}
+
+
+/* SWidget interface
+ *****************************************************************************/
+
+void SMediaPlayerEditorOutput::Tick(const FGeometry& AllottedGeometry, const double InCurrentTime, const float InDeltaTime)
+{
+	UpdateMaterial();
+	UpdateSoundWave();
+}
+
+
+/* SMediaPlayerEditorOutput implementation
+ *****************************************************************************/
+
+void SMediaPlayerEditorOutput::UpdateMaterial()
+{
+	UMediaTexture* Texture = MediaPlayer->GetVideoTexture();
+
+	if (Texture == nullptr)
+	{
+		// create default texture if needed
+		if (DefaultTexture == nullptr)
+		{
+			DefaultTexture = NewObject<UMediaTexture>(GetTransientPackage(), NAME_None, RF_Transient | RF_Public);
+			
+			if (DefaultTexture != nullptr)
+			{
+				DefaultTexture->UpdateResource();
+				DefaultTexture->AddToRoot();
+			}
+		}
+
+		// set default texture as output sink
+		if (DefaultTexture != nullptr)
+		{
+			TSharedPtr<IMediaPlayer> Player = MediaPlayer->GetPlayer();
+
+			if (Player.IsValid())
+			{
+				Player->GetOutput().SetVideoSink(DefaultTexture);
+			}
+		}
+
+		Texture = DefaultTexture;
+	}
+
+	// assign new texture to material
+	if (Texture != CurrentTexture)
+	{
+		TextureSampler->Texture = Texture;
+		Material->PostEditChange();
+		CurrentTexture = Texture;
+	}
+
+	// resize current texture
+	if (CurrentTexture != nullptr)
+	{
+		MaterialBrush->ImageSize.X = CurrentTexture->GetSurfaceWidth();
+		MaterialBrush->ImageSize.Y = CurrentTexture->GetSurfaceHeight();
+	}
+	else
+	{
+		MaterialBrush->ImageSize = FVector2D::ZeroVector;
+	}
+}
+
+
+void SMediaPlayerEditorOutput::UpdateSoundWave()
+{
+	if (AudioComponent == nullptr)
+	{
+		return;
+	}
+
+	if (GEditor->PlayWorld == nullptr)
+	{
+		UMediaSoundWave* SoundWave = MediaPlayer->GetSoundWave();
+
+		if (SoundWave == nullptr)
+		{
+			// set default sound wave as output sink
+			if (DefaultSoundWave != nullptr)
+			{
+				TSharedPtr<IMediaPlayer> Player = MediaPlayer->GetPlayer();
+
+				if (Player.IsValid())
+				{
+					Player->GetOutput().SetAudioSink(DefaultSoundWave);
+				}
+			}
+
+			SoundWave = DefaultSoundWave;
+		}
+
+		// assign new sound wave to audio component
+		if (SoundWave != CurrentSoundWave)
+		{
+			CurrentSoundWave = SoundWave;
+			AudioComponent->SetSound(SoundWave);
+		}
+	}
+	else
+	{
+		// no audio in PIE, PIV1, Simulate
+		AudioComponent->SetSound(nullptr);
+		CurrentSoundWave = nullptr;
+	}
+}
+
+
+/* SMediaPlayerEditorOutput callbacks
+ *****************************************************************************/
+
+void SMediaPlayerEditorOutput::HandleMediaPlayerMediaEvent(EMediaEvent Event)
+{
+	if ((AudioComponent == nullptr) || !CurrentSoundWave.IsValid())
+	{
+		return;
+	}
+
+	switch (Event)
+	{
+	case EMediaEvent::PlaybackEndReached:
+	case EMediaEvent::PlaybackSuspended:
+		AudioComponent->Stop();
+		break;
+
+	case EMediaEvent::PlaybackResumed:
+		AudioComponent->Play();
+		break;
+	}
+}
diff --git a/Engine/Plugins/Media/MediaPlayerEditor/Source/MediaPlayerEditor/Private/Widgets/SMediaPlayerEditorOutput.h b/Engine/Plugins/Media/MediaPlayerEditor/Source/MediaPlayerEditor/Private/Widgets/SMediaPlayerEditorOutput.h
new file mode 100644
--- /dev/null
+++ b/Engine/Plugins/Media/MediaPlayerEditor/Source/MediaPlayerEditor/Private/Widgets/SMediaPlayerEditorOutput.h
@@ -0,0 +1,91 @@
+// Copyright 1998-2016 Epic Games, Inc. All Rights Reserved.
+
+#pragma once
+
+
+class UMaterial;
+class UMaterialExpressionTextureSample;
+class UMediaPlayer;
+class UMediaTexture;
+
+
+/**
+ * Handles content output in the viewer tab in the UMediaPlayer asset editor.
+ */
+class SMediaPlayerEditorOutput
+	: public SCompoundWidget
+{
+public:
+
+	SLATE_BEGIN_ARGS(SMediaPlayerEditorOutput) { }
+	SLATE_END_ARGS()
+
+public:
+
+	/** Default constructor. */
+	SMediaPlayerEditorOutput();
+
+	/** Destructor. */
+	~SMediaPlayerEditorOutput();
+
+public:
+
+	/**
+	 * Construct this widget
+	 *
+	 * @param InArgs The declaration data for this widget.
+	 * @param InMediaPlayer The UMediaPlayer asset to show the details for.
+	 */
+	void Construct(const FArguments& InArgs, UMediaPlayer& InMediaPlayer);
+
+public:
+
+	//~ SWidget interface
+
+	virtual void Tick(const FGeometry& AllottedGeometry, const double InCurrentTime, const float InDeltaTime) override;
+
+protected:
+
+	/** Update the material to reflect the media player's latest video textures. */
+	void UpdateMaterial();
+
+	/** Update the sound wave to reflect the media player's latest sound output. */
+	void UpdateSoundWave();
+
+private:
+
+	/** Callback for media events from the media player. */
+	void HandleMediaPlayerMediaEvent(EMediaEvent Event);
+
+private:
+
+	/** The audio component used to play the current sound wave. */
+	UAudioComponent* AudioComponent;
+
+	/** The media sound wave currently being played. */
+	TWeakObjectPtr<UMediaSoundWave> CurrentSoundWave;
+
+	/** The media texture currently being shown. */
+	TWeakObjectPtr<UMediaTexture> CurrentTexture;
+
+	/** The media sound wave to use if the player doesn't have one assigned. */
+	UMediaSoundWave* DefaultSoundWave;
+
+	/** The media texture to use if the player doesn't have one assigned. */
+	UMediaTexture* DefaultTexture;
+
+	/** The material that wraps the video texture for display in an SImage. */
+	UMaterial* Material;
+
+	/** The Slate brush that renders the material. */
+	TSharedPtr<FSlateBrush> MaterialBrush;
+
+	/** The media player whose video texture is shown in this widget. */
+	UMediaPlayer* MediaPlayer;
+
+	/** Whether Play-In-Editor is currently active. */
+	bool PieActive;
+
+	/** The video texture sampler in the wrapper material. */
+	UMaterialExpressionTextureSample* TextureSampler;
+};
diff --git a/Engine/Plugins/Media/MediaPlayerEditor/Source/MediaPlayerEditor/Private/Widgets/SMediaPlayerEditorPlaylist.cpp b/Engine/Plugins/Media/MediaPlayerEditor/Source/MediaPlayerEditor/Private/Widgets/SMediaPlayerEditorPlaylist.cpp
new file mode 100644
--- /dev/null
+++ b/Engine/Plugins/Media/MediaPlayerEditor/Source/MediaPlayerEditor/Private/Widgets/SMediaPlayerEditorPlaylist.cpp
@@ -0,0 +1,147 @@
+// Copyright 1998-2016 Epic Games, Inc. All Rights Reserved.
+
+#include "MediaPlayerEditorPCH.h"
+#include "SMediaPlayerEditorPlaylist.h"
+#include "SMediaSourceTableRow.h"
+
+
+#define LOCTEXT_NAMESPACE "SMediaPlayerEditorMedia"
+
+
+/* SMediaPlayerEditorPlaylist structors
+ *****************************************************************************/
+
+SMediaPlayerEditorPlaylist::SMediaPlayerEditorPlaylist()
+	: MediaPlayer(nullptr)
+{ }
+
+
+SMediaPlayerEditorPlaylist::~SMediaPlayerEditorPlaylist()
+{
+	if (MediaPlayer != nullptr)
+	{
+		MediaPlayer->OnMediaEvent().RemoveAll(this);
+	}
+
+	FCoreUObjectDelegates::OnObjectPropertyChanged.RemoveAll(this);
+}
+
+
+/* SMediaPlayerEditorPlaylist interface
+ *****************************************************************************/
+
+void SMediaPlayerEditorPlaylist::Construct(const FArguments& InArgs, UMediaPlayer& InMediaPlayer, const TSharedRef<ISlateStyle>& InStyle)
+{
+	MediaPlayer = &InMediaPlayer;
+	Style = InStyle;
+
+	ChildSlot
+	[
+		SNew(SVerticalBox)
+
+		+ SVerticalBox::Slot()
+			.FillHeight(1.0f)
+			[
+				SNew(SBorder)
+					.BorderImage(FEditorStyle::GetBrush("ToolPanel.GroupBorder"))
+					.Padding(0.0f)
+					[
+						// message list
+						SAssignNew(MediaSourceListView, SListView<TSharedPtr<FMediaSourceTableEntry>>)
+							.ItemHeight(24.0f)
+							.ListItemsSource(&MediaSourceList)
+							.SelectionMode(ESelectionMode::Single)
+							.OnGenerateRow(this, &SMediaPlayerEditorPlaylist::HandleMediaSourceListGenerateRow)
+							.OnMouseButtonDoubleClick(this, &SMediaPlayerEditorPlaylist::HandleMediaSourceListDoubleClick)
+							.HeaderRow
+							(
+								SNew(SHeaderRow)
+
+								+ SHeaderRow::Column("Icon")
+									.DefaultLabel(LOCTEXT("PlaylistIconColumnHeader", " "))
+									.FixedWidth(12.0f)
+
+								+ SHeaderRow::Column("Index")
+									.DefaultLabel(LOCTEXT("PlaylistIndexColumnHeader", "#"))
+									.FillWidth(0.1f)
+
+								+ SHeaderRow::Column("Source")
+									.DefaultLabel(LOCTEXT("PlaylistSourceColumnHeader", "Media Source"))
+									.FillWidth(0.5f)
+
+								+ SHeaderRow::Column("Type")
+									.DefaultLabel(LOCTEXT("PlaylistTypeColumnHeader", "Type"))
+									.FillWidth(0.4f)
+							)
+					]
+			]
+	];
+
+	FCoreUObjectDelegates::OnObjectPropertyChanged.AddSP(this, &SMediaPlayerEditorPlaylist::HandleCoreObjectPropertyChanged);
+	MediaPlayer->OnMediaEvent().AddSP(this, &SMediaPlayerEditorPlaylist::HandleMediaPlayerMediaEvent);
+
+	ReloadMediaSourceList();
+}
+
+
+/* SMediaPlayerEditorPlaylist implementation
+ *****************************************************************************/
+
+void SMediaPlayerEditorPlaylist::ReloadMediaSourceList()
+{
+	MediaSourceList.Reset();
+
+	UMediaPlaylist* Playlist = MediaPlayer->GetPlaylist();
+
+	if (Playlist != nullptr)
+	{
+		for (int32 EntryIndex = 0; EntryIndex < Playlist->Num(); ++EntryIndex)
+		{
+			MediaSourceList.Add(MakeShareable(new FMediaSourceTableEntry(EntryIndex, Playlist->Get(EntryIndex))));
+		}
+	}
+
+	MediaSourceListView->RequestListRefresh();
+}
+
+
+/* SMediaPlayerEditorPlaylist callbacks
+ *****************************************************************************/
+
+void SMediaPlayerEditorPlaylist::HandleCoreObjectPropertyChanged(UObject* Object, struct FPropertyChangedEvent& ChangedEvent)
+{
+	if ((Object != nullptr) && (Object == MediaPlayer->GetPlaylist()))
+	{
+		ReloadMediaSourceList();
+	}
+}
+
+
+void SMediaPlayerEditorPlaylist::HandleMediaPlayerMediaEvent(EMediaEvent Event)
+{
+	if ((Event == EMediaEvent::MediaClosed) || (Event == EMediaEvent::MediaOpened))
+	{
+		ReloadMediaSourceList();
+	}
+}
+
+
+void SMediaPlayerEditorPlaylist::HandleMediaSourceListDoubleClick(TSharedPtr<FMediaSourceTableEntry> InItem)
+{
+	if (InItem.IsValid())
+	{
+		MediaPlayer->OpenPlaylistIndex(MediaPlayer->GetPlaylist(), InItem->Index);
+	}
+}
+
+
+TSharedRef<ITableRow> SMediaPlayerEditorPlaylist::HandleMediaSourceListGenerateRow(TSharedPtr<FMediaSourceTableEntry> Entry, const TSharedRef<STableViewBase>& OwnerTable)
+{
+	return SNew(SMediaSourceTableRow, OwnerTable)
+		.Entry(Entry)
+		.Opened_Lambda([=]() { return MediaPlayer->GetPlaylistIndex() == Entry->Index; })
+		.Style(Style);
+}
+
+
+#undef LOCTEXT_NAMESPACE
diff --git a/Engine/Plugins/Media/MediaPlayerEditor/Source/MediaPlayerEditor/Private/Widgets/SMediaPlayerEditorViewer.cpp b/Engine/Plugins/Media/MediaPlayerEditor/Source/MediaPlayerEditor/Private/Widgets/SMediaPlayerEditorViewer.cpp
--- a/Engine/Plugins/Media/MediaPlayerEditor/Source/MediaPlayerEditor/Private/Widgets/SMediaPlayerEditorViewer.cpp
+++ b/Engine/Plugins/Media/MediaPlayerEditor/Source/MediaPlayerEditor/Private/Widgets/SMediaPlayerEditorViewer.cpp
@@ -1,8 +1,9 @@
 // Copyright 1998-2016 Epic Games, Inc. All Rights Reserved.
 
-#include "MediaPlayerEditorPrivatePCH.h"
-#include "IMediaPlayer.h"
-#include "SNumericEntryBox.h"
+#include "MediaPlayerEditorPCH.h"
+#include "MediaPlayerEditorSettings.h"
+#include "SMediaPlayerEditorOutput.h"
+#include "SMediaPlayerEditorViewer.h"
 
 
 #define LOCTEXT_NAMESPACE "SMediaPlayerEditorViewer"
@@ -12,8 +13,7 @@
  *****************************************************************************/
 
 SMediaPlayerEditorViewer::SMediaPlayerEditorViewer()
-	: CaptionBuffer(MakeShareable(new FMediaSampleBuffer))
-	, MediaPlayer(nullptr)
+	: MediaPlayer(nullptr)
 { }
 
 
@@ -21,8 +21,7 @@ SMediaPlayerEditorViewer::~SMediaPlayerEditorViewer()
 {
 	if (MediaPlayer != nullptr)
 	{
-		MediaPlayer->OnMediaChanged().RemoveAll(this);
-		MediaPlayer->OnTracksChanged().RemoveAll(this);
+		MediaPlayer->OnMediaEvent().RemoveAll(this);
 	}
 }
 
@@ -30,10 +29,116 @@ SMediaPlayerEditorViewer::~SMediaPlayerEditorViewer()
 /* SMediaPlayerEditorPlayer interface
  *****************************************************************************/
 
-void SMediaPlayerEditorViewer::Construct(const FArguments& InArgs, UMediaPlayer* InMediaPlayer, const TSharedRef<ISlateStyle>& InStyle)
+void SMediaPlayerEditorViewer::Construct(const FArguments& InArgs, UMediaPlayer& InMediaPlayer, const TSharedRef<ISlateStyle>& InStyle)
 {
-	MediaPlayer = InMediaPlayer;
-	Viewport = MakeShareable(new FMediaPlayerEditorViewport());
+	MediaPlayer = &InMediaPlayer;
+	Style = InStyle;
+
+	// initialize media player asset
+	MediaPlayer->OnMediaEvent().AddSP(this, &SMediaPlayerEditorViewer::HandleMediaPlayerMediaEvent);
+	MediaPlayer->DesiredPlayerName = NAME_None;
+
+	FName DesiredPlayerName = GetDefault<UMediaPlayerEditorSettings>()->DesiredPlayerName;
+
+	if (DesiredPlayerName != NAME_None)
+	{
+		IMediaModule* MediaModule = FModuleManager::LoadModulePtr<IMediaModule>("Media");
+
+		if ((MediaModule != nullptr) && (MediaModule->GetPlayerFactory(DesiredPlayerName) != nullptr))
+		{
+			MediaPlayer->DesiredPlayerName = DesiredPlayerName;
+		}
+	}
+
+	// initialize settings menu
+	FMenuBuilder SettingsMenuBuilder(true, nullptr);
+	{
+		SettingsMenuBuilder.BeginSection("PlayerSection", LOCTEXT("PlayerSection", "Player"));
+		{
+			SettingsMenuBuilder.AddSubMenu(
+				LOCTEXT("DecoderMenuLabel", "Decoder"),
+				LOCTEXT("DecoderMenuTooltip", "Select the desired media decoder"),
+				FNewMenuDelegate::CreateRaw(this, &SMediaPlayerEditorViewer::HandlePlayerMenuNewMenu),
+				false,
+				FSlateIcon()
+			);
+		}
+		SettingsMenuBuilder.EndSection();
+
+		SettingsMenuBuilder.BeginSection("TracksSection", LOCTEXT("TracksSection", "Tracks"));
+		{
+			SettingsMenuBuilder.AddSubMenu(
+				LOCTEXT("AudioTrackMenuLabel", "Audio Track"),
+				LOCTEXT("AudioTrackMenuTooltip", "Select the active audio track"),
+				FNewMenuDelegate::CreateRaw(this, &SMediaPlayerEditorViewer::HandleTrackMenuNewMenu, EMediaPlayerTrack::Audio),
+				false,
+				FSlateIcon()
+			);
+
+			SettingsMenuBuilder.AddSubMenu(
+				LOCTEXT("CaptionTrackMenuLabel", "Caption Track"),
+				LOCTEXT("CaptionTrackMenuTooltip", "Select the active caption track"),
+				FNewMenuDelegate::CreateRaw(this, &SMediaPlayerEditorViewer::HandleTrackMenuNewMenu, EMediaPlayerTrack::Caption),
+				false,
+				FSlateIcon()
+			);
+
+			SettingsMenuBuilder.AddSubMenu(
+				LOCTEXT("ImageTrackMenuLabel", "Image Track"),
+				LOCTEXT("ImageTrackMenuTooltip", "Select the active image track"),
+				FNewMenuDelegate::CreateRaw(this, &SMediaPlayerEditorViewer::HandleTrackMenuNewMenu, EMediaPlayerTrack::Image),
+				false,
+				FSlateIcon()
+			);
+
+			SettingsMenuBuilder.AddSubMenu(
+				LOCTEXT("VideoTrackMenuLabel", "Video Track"),
+				LOCTEXT("VideoTrackMenuTooltip", "Select the active video track"),
+				FNewMenuDelegate::CreateRaw(this, &SMediaPlayerEditorViewer::HandleTrackMenuNewMenu, EMediaPlayerTrack::Video),
+				false,
+				FSlateIcon()
+			);
+		}
+		SettingsMenuBuilder.EndSection();
+
+		SettingsMenuBuilder.BeginSection("ViewSection", LOCTEXT("ViewSection", "View"));
+		{
+			SettingsMenuBuilder.AddMenuEntry(
+				LOCTEXT("ScaleToFitMenuLabel", "Scale Video To Fit"),
+				LOCTEXT("ScaleToFitMenuTooltip", "Scale the video to fit the viewport and maintain the aspect ratio"),
+				FSlateIcon(),
+				FUIAction(
+					FExecuteAction::CreateLambda([]{
+						auto Settings = GetMutableDefault<UMediaPlayerEditorSettings>();
+						Settings->ScaleVideoToFit = !Settings->ScaleVideoToFit;
+						Settings->SaveConfig();
+					}),
+					FCanExecuteAction(),
+					FIsActionChecked::CreateLambda([]{ return GetDefault<UMediaPlayerEditorSettings>()->ScaleVideoToFit; })
+				),
+				NAME_None,
+				EUserInterfaceActionType::ToggleButton
+			);
+
+			SettingsMenuBuilder.AddMenuEntry(
+				LOCTEXT("ShowCaptionsMenuLabel", "Show Captions"),
+				LOCTEXT("ShowCaptionsMenuTooltip", "Show caption and subtitle text"),
+				FSlateIcon(),
+				FUIAction(
+					FExecuteAction::CreateLambda([]{
+						auto Settings = GetMutableDefault<UMediaPlayerEditorSettings>();
+						Settings->ShowCaptions = !Settings->ShowCaptions;
+						Settings->SaveConfig();
+					}),
+					FCanExecuteAction(),
+					FIsActionChecked::CreateLambda([]{ return GetDefault<UMediaPlayerEditorSettings>()->ShowCaptions; })
+				),
+				NAME_None,
+				EUserInterfaceActionType::ToggleButton
+			);
+		}
+		SettingsMenuBuilder.EndSection();
+	}
 
 	ChildSlot
 	[
@@ -41,340 +146,255 @@ void SMediaPlayerEditorViewer::Construct(const FArguments& InArgs, UMediaPlayer*
 
 		+ SVerticalBox::Slot()
 			.AutoHeight()
-			.Padding(4.0f, 0.0f)
 			[
-				SNew(SOverlay)
+				SNew(SHorizontalBox)
 
-				+ SOverlay::Slot()
-					.VAlign(VAlign_Center)
+				+ SHorizontalBox::Slot()
+					.FillWidth(1.0f)
 					[
-						SNew(SHorizontalBox)
-							.Visibility(this, &SMediaPlayerEditorViewer::HandleNoMediaSelectedTextVisibility)
+						// url box
+						SAssignNew(UrlTextBox, SEditableTextBox)
+							.ClearKeyboardFocusOnCommit(true)
+							.HintText(LOCTEXT("UrlTextBoxHint", "Media URL"))
+							.ToolTipText(LOCTEXT("UrlTextBoxToolTip", "Enter the URL of a media source"))
+							.OnKeyDownHandler(this, &SMediaPlayerEditorViewer::HandleUrlBoxKeyDown)
+							.OnTextCommitted_Lambda([this](const FText& InText, ETextCommit::Type InCommitType){
+								if (InCommitType == ETextCommit::OnEnter)
+								{
+									FString Url = UrlTextBox->GetText().ToString();
+									if (!Url.Contains(TEXT("://"), ESearchCase::CaseSensitive))
+									{
+										Url.InsertAt(0, TEXT("file://"));
+									}
+									MediaPlayer->OpenUrl(Url);
+								}
+							})
+					]
 
-						+ SHorizontalBox::Slot()
-							.AutoWidth()
-							.VAlign(VAlign_Center)
+				+ SHorizontalBox::Slot()
+					.AutoWidth()
+					.Padding(4.0f, 0.0f, 0.0f, 0.0f)
+					[
+						// go button
+						SNew(SButton)
+							.ToolTipText_Lambda([this]() {
+								return ((UrlTextBox->GetText().ToString() == MediaPlayer->GetUrl()) && !MediaPlayer->GetUrl().IsEmpty())
+									? LOCTEXT("ReloadButtonToolTip", "Reload the current media URL")
+									: LOCTEXT("GoButtonToolTip", "Open the specified media URL");
+							})
+							.IsEnabled_Lambda([this]{
+								return !UrlTextBox->GetText().IsEmpty();
+							})
+							.OnClicked_Lambda([this]{
+								MediaPlayer->OpenUrl(UrlTextBox->GetText().ToString());
+								return FReply::Handled();
+							})
 							[
 								SNew(SImage)
-									.Image(FCoreStyle::Get().GetBrush("Icons.Error"))
-							]
-
-						+ SHorizontalBox::Slot()
-							.FillWidth(1.0f)
-							.Padding(4.0f, 0.0f, 0.0f, 0.0f)
-							.VAlign(VAlign_Center)
-							[
-								SNew(STextBlock)
-									.Text(LOCTEXT("NoMediaSelectedText", "Please pick a media source in the Details panel!"))
+									.Image_Lambda([this]() {
+										return ((UrlTextBox->GetText().ToString() == MediaPlayer->GetUrl()) && !MediaPlayer->GetUrl().IsEmpty())
+											? Style->GetBrush("MediaPlayerEditor.ReloadButton")
+											: Style->GetBrush("MediaPlayerEditor.GoButton");
+									})
 							]
 					]
+			]
 
-				+ SOverlay::Slot()
+		+ SVerticalBox::Slot()
+			.FillHeight(1.0f)
+			.Padding(0.0f, 4.0f, 0.0f, 0.0f)
+			[
+				SNew(SBorder)
+					.BorderImage(FCoreStyle::Get().GetBrush("BlackBrush"))
+					.Padding(0.0f)
+					.OnMouseButtonUp(this, &SMediaPlayerEditorViewer::HandleTextureMouseButtonUp)
 					[
-						SNew(SHorizontalBox)
-							.Visibility(this, &SMediaPlayerEditorViewer::HandleTrackSelectionBoxVisibility)
+						SNew(SOverlay)
 
-						+ SHorizontalBox::Slot()
-							.AutoWidth()
+						+ SOverlay::Slot()
+							[
+								// movie texture
+								SNew(SScaleBox)
+									.Stretch_Lambda([]() -> EStretch::Type { return GetDefault<UMediaPlayerEditorSettings>()->ScaleVideoToFit ? EStretch::ScaleToFit : EStretch::Fill; })
+									[
+										SNew(SMediaPlayerEditorOutput, InMediaPlayer)
+									]
+							]
+
+						+ SOverlay::Slot()
+							.Padding(FMargin(12.0f, 8.0f))
 							[
-								// video track selector
 								SNew(SHorizontalBox)
 
 								+ SHorizontalBox::Slot()
-									.AutoWidth()
-									.VAlign(VAlign_Center)
+									.HAlign(HAlign_Left)
+									.VAlign(VAlign_Top)
 									[
+										// playback state
 										SNew(STextBlock)
-											.Text(LOCTEXT("VideoTrackLabel", "Video Track:"))
+											.ColorAndOpacity(FSlateColor::UseSubduedForeground())
+											.Font(FSlateFontInfo(FPaths::EngineContentDir() / TEXT("Slate/Fonts/Roboto-Regular.ttf"), 18))
+											.ShadowOffset(FVector2D(1.f, 1.f))
+											.Text(this, &SMediaPlayerEditorViewer::HandleOverlayStateText)
 									]
 
 								+ SHorizontalBox::Slot()
-									.FillWidth(1.0f)
-									.Padding(4.0f, 0.0f)
-									.VAlign(VAlign_Center)
+									.HAlign(HAlign_Right)
+									.VAlign(VAlign_Top)
 									[
-										SAssignNew(VideoTrackComboBox, SComboBox<IMediaVideoTrackPtr>)
-											.OnGenerateWidget(this, &SMediaPlayerEditorViewer::HandleVideoTrackComboBoxGenerateWidget)
-											.OnSelectionChanged(this, &SMediaPlayerEditorViewer::HandleVideoTrackComboBoxSelectionChanged)
-											.OptionsSource(&VideoTracks)
-											.ToolTipText(LOCTEXT("CaptionTrackToolTip", "Select the video track to be played"))
-											[
-												SNew(STextBlock)
-													.Text(this, &SMediaPlayerEditorViewer::HandleVideoTrackComboBoxText)
-											]
+										// player name
+										SNew(STextBlock)
+											.ColorAndOpacity(FSlateColor::UseSubduedForeground())
+											.Font(FSlateFontInfo(FPaths::EngineContentDir() / TEXT("Slate/Fonts/Roboto-Regular.ttf"), 18))
+											.ShadowOffset(FVector2D(1.f, 1.f))
+											.Text(this, &SMediaPlayerEditorViewer::HandleOverlayPlayerNameText)
 									]
 							]
 
-						+ SHorizontalBox::Slot()
-							.AutoWidth()
-							.Padding(16.0f, 0.0f, 0.0f, 0.0f)
+						+ SOverlay::Slot()
+							.HAlign(HAlign_Fill)
+							.VAlign(VAlign_Bottom)
+							.Padding(FMargin(50.0f, 20.0f))
 							[
-								// audio track selector
-								SNew(SHorizontalBox)
+								// caption text
+								SNew(STextBlock)
+									.AutoWrapText(true)
+									.Font(FSlateFontInfo(FPaths::EngineContentDir() / TEXT("Slate/Fonts/Roboto-Regular.ttf"), 14))
+									.ShadowOffset(FVector2D(1.f, 1.f))
+									.Text(this, &SMediaPlayerEditorViewer::HandleOverlayCaptionText)
+							]
+					]
+			]
 
-								+ SHorizontalBox::Slot()
-									.AutoWidth()
-									.VAlign(VAlign_Center)
+		+ SVerticalBox::Slot()
+			.AutoHeight()
+			.Padding(0.0f, 2.0f, 0.0f, 0.0f)
+			.VAlign(VAlign_Center)
+			[
+				SNew(SBorder)
+					.BorderImage(FEditorStyle::GetBrush("ToolPanel.GroupBorder"))
+					.ForegroundColor(FLinearColor::Gray)
+					.Padding(6.0f)
+					[
+						SNew(SVerticalBox)
+
+						+ SVerticalBox::Slot()
+							.AutoHeight()
+							[
+								SNew(SOverlay)
+
+								+ SOverlay::Slot()
 									[
-										SNew(STextBlock)
-											.Text(LOCTEXT("AudioTrackLabel", "Audio Track:"))
+										// scrubber
+										SAssignNew(ScrubberSlider, SSlider)
+											.IsEnabled(this, &SMediaPlayerEditorViewer::HandlePositionSliderIsEnabled)
+											.OnMouseCaptureBegin(this, &SMediaPlayerEditorViewer::HandlePositionSliderMouseCaptureBegin)
+											.OnMouseCaptureEnd(this, &SMediaPlayerEditorViewer::HandlePositionSliderMouseCaptureEnd)
+											.OnValueChanged(this, &SMediaPlayerEditorViewer::HandlePositionSliderValueChanged)
+											.Orientation(Orient_Horizontal)
+											.Value(this, &SMediaPlayerEditorViewer::HandlePositionSliderValue)
 									]
 
-								+ SHorizontalBox::Slot()
-									.FillWidth(1.0f)
-									.Padding(4.0f, 0.0f, 0.0f, 0.0f)
-									.VAlign(VAlign_Center)
+								+ SOverlay::Slot()
 									[
-										SAssignNew(AudioTrackComboBox, SComboBox<IMediaAudioTrackPtr>)
-											.OnGenerateWidget(this, &SMediaPlayerEditorViewer::HandleAudioTrackComboBoxGenerateWidget)
-											.OnSelectionChanged(this, &SMediaPlayerEditorViewer::HandleAudioTrackComboBoxSelectionChanged)
-											.OptionsSource(&AudioTracks)
-											.ToolTipText(LOCTEXT("CaptionTrackToolTip", "Select the audio track to be played"))
-											[
-												SNew(STextBlock)
-													.Text(this, &SMediaPlayerEditorViewer::HandleAudioTrackComboBoxText)
-											]
+										SNew(SProgressBar)
+											.ToolTipText(LOCTEXT("BufferingTooltip", "Buffering..."))
+											.Visibility_Lambda([this]() -> EVisibility { return MediaPlayer->IsPreparing() ? EVisibility::Visible : EVisibility::Hidden; })
 									]
 							]
 
-						+ SHorizontalBox::Slot()
-							.AutoWidth()
-							.Padding(16.0f, 0.0f, 0.0f, 0.0f)
+						+ SVerticalBox::Slot()
+							.AutoHeight()
+							.Padding(0.0f, 2.0f, 0.0f, 0.0f)
 							[
-								// caption track selector
 								SNew(SHorizontalBox)
 
+								// timer
 								+ SHorizontalBox::Slot()
 									.AutoWidth()
+									.Padding(0.0f, 0.0f, 4.0f, 0.0f)
 									.VAlign(VAlign_Center)
 									[
 										SNew(STextBlock)
-											.Text(LOCTEXT("CaptionTrackLabel", "Caption Track:"))
+											.Text(this, &SMediaPlayerEditorViewer::HandleTimerTextBlockText)
+											.ToolTipText(this, &SMediaPlayerEditorViewer::HandleTimerTextBlockToolTipText)
 									]
 
+								// transport controls
 								+ SHorizontalBox::Slot()
 									.FillWidth(1.0f)
-									.Padding(4.0f, 0.0f, 0.0f, 0.0f)
+									.Padding(8.0f, 0.0f)
 									.VAlign(VAlign_Center)
 									[
-										SAssignNew(CaptionTrackComboBox, SComboBox<IMediaCaptionTrackPtr>)
-											.OnGenerateWidget(this, &SMediaPlayerEditorViewer::HandleCaptionTrackComboBoxGenerateWidget)
-											.OnSelectionChanged(this, &SMediaPlayerEditorViewer::HandleCaptionTrackComboBoxSelectionChanged)
-											.OptionsSource(&CaptionTracks)
-											.ToolTipText(LOCTEXT("CaptionTrackToolTip", "Select the caption text to be displayed"))
-											[
-												SNew(STextBlock)
-													.Text(this, &SMediaPlayerEditorViewer::HandleCaptionTrackComboBoxText)
-											]
+										SNew(SSpacer)
+										//SNew(STransportControl)
 									]
-							]
-
-						+ SHorizontalBox::Slot()
-							.FillWidth(1.0f)
-							.HAlign(HAlign_Right)
-							.Padding(16.0f, 0.0f, 0.0f, 0.0f)
-							[
-								// playback rate spin box
-								SNew(SHorizontalBox)
 
+								// settings
 								+ SHorizontalBox::Slot()
 									.AutoWidth()
-									.VAlign(VAlign_Center)
-									[
-										SNew(STextBlock)
-											.Text(LOCTEXT("PlaybackRateLabel", "Playback Rate:"))
-									]
-
-								+ SHorizontalBox::Slot()
-									.FillWidth(1.0f)
 									.Padding(4.0f, 0.0f, 0.0f, 0.0f)
 									.VAlign(VAlign_Center)
 									[
-										SNew(SNumericEntryBox<float>)
-											.Delta(1.0f)
-											.MaxValue(this, &SMediaPlayerEditorViewer::HandlePlaybackRateBoxMaxValue)
-											.MinValue(this, &SMediaPlayerEditorViewer::HandlePlaybackRateBoxMinValue)
-											.Value(this, &SMediaPlayerEditorViewer::HandlePlaybackRateSpinBoxValue)
-											.OnValueChanged(this, &SMediaPlayerEditorViewer::HandlePlaybackRateBoxValueChanged)
+										SNew(SComboButton)
+											.ContentPadding(0.0f)
+											.ButtonContent()
+											[
+												SNew(SHorizontalBox)
+
+												+ SHorizontalBox::Slot()
+													.AutoWidth()
+													.VAlign(VAlign_Center)
+													[
+														SNew(SImage)
+															.Image(InStyle->GetBrush("MediaPlayerEditor.SettingsButton"))
+													]
+
+												+ SHorizontalBox::Slot()
+													.AutoWidth()
+													.Padding(3.0f, 0.0f, 0.0f, 0.0f)
+													.VAlign(VAlign_Center)
+													[
+														SNew(STextBlock)
+															.Text(LOCTEXT("OptionsButton", "Playback Options"))
+													]
+											]
+											.ButtonStyle(FEditorStyle::Get(), "ToggleButton")
+											.ForegroundColor(FSlateColor::UseForeground())
+											.MenuContent()
+											[
+												SettingsMenuBuilder.MakeWidget()
+											]
 									]
 							]
 					]
 			]
-
-		+ SVerticalBox::Slot()
-			.FillHeight(1.0f)
-			.Padding(0.0f, 4.0f, 0.0f, 8.0f)
-			[
-				SNew(SOverlay)
-
-				+ SOverlay::Slot()
-					[
-						// movie viewport
-						SNew(SViewport)
-							.EnableGammaCorrection(false)
-							.ViewportInterface(Viewport)
-					]
-
-				+ SOverlay::Slot()
-					.HAlign(HAlign_Left)
-					.VAlign(VAlign_Top)
-					.Padding(FMargin(12.0f, 8.0f))
-					[
-						// playback state
-						SNew(STextBlock)
-							.Font(FSlateFontInfo(FPaths::EngineContentDir() / TEXT("Slate/Fonts/Roboto-Regular.ttf"), 20))
-							.ShadowOffset(FVector2D(1.f, 1.f))
-							.Text(this, &SMediaPlayerEditorViewer::HandleOverlayStateText)
-					]
-
-				+ SOverlay::Slot()
-					.HAlign(HAlign_Center)
-					.VAlign(VAlign_Bottom)
-					.Padding(16.0f)
-					[
-						// caption text
-						SNew(STextBlock)
-							.Font(FSlateFontInfo(FPaths::EngineContentDir() / TEXT("Slate/Fonts/Roboto-Regular.ttf"), 16))
-							.ShadowOffset(FVector2D(1.f, 1.f))
-							.Text(this, &SMediaPlayerEditorViewer::HandleOverlayCaptionText)
-							.WrapTextAt(600.0f)
-					]
-			]
-
-		+ SVerticalBox::Slot()
-			.AutoHeight()
-			.Padding(4.0f, 0.0f)
-			.VAlign(VAlign_Center)
-			[
-				SNew(SHorizontalBox)
-
-				// elapsed time
-				+ SHorizontalBox::Slot()
-					.AutoWidth()
-					[
-						SNew(STextBlock)
-							.Text(this, &SMediaPlayerEditorViewer::HandleElapsedTimeTextBlockText)
-							.ToolTipText(LOCTEXT("ElapsedTimeTooltip", "Elapsed Time"))
-					]
-
-				// scrubber
-				+ SHorizontalBox::Slot()
-					.FillWidth(1.0f)
-					.Padding(8.0f, 0.0f)
-					[
-						SAssignNew(ScrubberSlider, SSlider)
-							.IsEnabled(this, &SMediaPlayerEditorViewer::HandlePositionSliderIsEnabled)
-							.OnMouseCaptureBegin(this, &SMediaPlayerEditorViewer::HandlePositionSliderMouseCaptureBegin)
-							.OnMouseCaptureEnd(this, &SMediaPlayerEditorViewer::HandlePositionSliderMouseCaptureEnd)
-							.OnValueChanged(this, &SMediaPlayerEditorViewer::HandlePositionSliderValueChanged)
-							.Orientation(Orient_Horizontal)
-							.Value(this, &SMediaPlayerEditorViewer::HandlePositionSliderValue)
-					]
-
-				// remaining time
-				+ SHorizontalBox::Slot()
-					.AutoWidth()
-					[
-						SNew(STextBlock)
-							.Text(this, &SMediaPlayerEditorViewer::HandleRemainingTimeTextBlockText)
-							.ToolTipText(LOCTEXT("RemainingTimeTooltip", "Remaining Time"))
-					]
-			]
 	];
 
-	MediaPlayer->OnMediaChanged().AddSP(this, &SMediaPlayerEditorViewer::HandleMediaPlayerMediaChanged);
-	MediaPlayer->OnTracksChanged().AddSP(this, &SMediaPlayerEditorViewer::HandleMediaPlayerTracksChanged);
-
 	RegisterActiveTimer(0.0f, FWidgetActiveTimerDelegate::CreateSP(this, &SMediaPlayerEditorViewer::HandleActiveTimer));
-	ReloadTracks();
 }
 
 
 /* SMediaPlayerEditorPlayer implementation
  *****************************************************************************/
 
-void SMediaPlayerEditorViewer::ReloadTracks()
+void SMediaPlayerEditorViewer::SetDesiredPlayerName(FName PlayerName)
 {
-	// clear track collections
-	IMediaCaptionTrackPtr SelectedCaptionTrack = CaptionTrackComboBox->GetSelectedItem();
-
-	if (SelectedCaptionTrack.IsValid())
-	{
-		SelectedCaptionTrack->GetStream().RemoveSink(CaptionBuffer);
-	}
-
-	AudioTracks.Reset();
-	CaptionTracks.Reset();
-	VideoTracks.Reset();
-
-	// get track collections
-	if (MediaPlayer != nullptr)
-	{
-		IMediaPlayerPtr Player = MediaPlayer->GetPlayer();
-
-		if (Player.IsValid())
-		{
-			for (auto AudioTrack : Player->GetAudioTracks())
-			{
-				AudioTracks.Add(AudioTrack);
-			}
-
-			for (auto CaptionTrack : Player->GetCaptionTracks())
-			{
-				CaptionTracks.Add(CaptionTrack);
-			}
-
-			for (auto VideoTrack : Player->GetVideoTracks())
-			{
-				VideoTracks.Add(VideoTrack);
-			}
-		}
-	}
-
-	// refresh combo box selections
-	if (!AudioTracks.Contains(AudioTrackComboBox->GetSelectedItem()))
+	if (PlayerName != MediaPlayer->DesiredPlayerName)
 	{
-		if (AudioTracks.Num() > 0)
-		{
-			AudioTrackComboBox->SetSelectedItem(AudioTracks[0]);
-		}
-		else
-		{
-			AudioTrackComboBox->SetSelectedItem(nullptr);
-		}		
-	}
+		MediaPlayer->DesiredPlayerName = PlayerName;
 
-	if (!CaptionTracks.Contains(CaptionTrackComboBox->GetSelectedItem()))
-	{
-		if (CaptionTracks.Num() > 0)
+		if (PlayerName != NAME_None)
 		{
-			CaptionTrackComboBox->SetSelectedItem(CaptionTracks[0]);
-			CaptionTracks[0]->GetStream().AddSink(CaptionBuffer);
+			MediaPlayer->Reopen();
 		}
-		else
-		{
-			CaptionTrackComboBox->SetSelectedItem(nullptr);
-		}		
 	}
 
-	if (!VideoTracks.Contains(VideoTrackComboBox->GetSelectedItem()))
+	auto Settings = GetMutableDefault<UMediaPlayerEditorSettings>();
 	{
-		if (VideoTracks.Num() > 0)
-		{
-			VideoTrackComboBox->SetSelectedItem(VideoTracks[0]);
-		}
-		else
-		{
-			VideoTrackComboBox->SetSelectedItem(nullptr);
-		}		
+		Settings->DesiredPlayerName = PlayerName;
+		Settings->SaveConfig();
 	}
-
-	AudioTrackComboBox->RefreshOptions();
-	CaptionTrackComboBox->RefreshOptions();
-	VideoTrackComboBox->RefreshOptions();
-
-	// refresh viewport
-	Viewport->Initialize(VideoTrackComboBox->GetSelectedItem());
 }
 
 
@@ -387,108 +407,66 @@ EActiveTimerReturnType SMediaPlayerEditorViewer::HandleActiveTimer(double InCurr
 }
 
 
-TSharedRef<SWidget> SMediaPlayerEditorViewer::HandleAudioTrackComboBoxGenerateWidget(IMediaAudioTrackPtr Value) const
+void SMediaPlayerEditorViewer::HandleMediaPlayerMediaEvent(EMediaEvent Event)
 {
-	return SNew(STextBlock)
-		.Text(Value->GetStream().GetDisplayName());
-}
-
-
-void SMediaPlayerEditorViewer::HandleAudioTrackComboBoxSelectionChanged(IMediaAudioTrackPtr Selection, ESelectInfo::Type SelectInfo)
-{
-	if (SelectInfo != ESelectInfo::Direct)
+	switch (Event)
 	{
-		ReloadTracks();
+	case EMediaEvent::MediaClosed:
+	case EMediaEvent::MediaOpened:
+	case EMediaEvent::MediaOpenFailed:
+		UrlTextBox->SetText(FText::FromString(MediaPlayer->GetUrl()));
+		break;
 	}
-}
 
+	if (Event == EMediaEvent::MediaOpenFailed)
+	{
+		UrlTextBox->SetBorderBackgroundColor(FLinearColor::Red);
 
-FText SMediaPlayerEditorViewer::HandleAudioTrackComboBoxText() const
-{
-	IMediaAudioTrackPtr AudioTrack = AudioTrackComboBox->GetSelectedItem();
+		FNotificationInfo NotificationInfo(LOCTEXT("MediaOpenFailedError", "The media failed to open. Check Output Log for details!"));
+		{
+			NotificationInfo.ExpireDuration = 2.0f;
+		}
 
-	if (AudioTrack.IsValid())
-	{
-		return AudioTrack->GetStream().GetDisplayName();
+		FSlateNotificationManager::Get().AddNotification(NotificationInfo)->SetCompletionState(SNotificationItem::CS_Fail);
 	}
-
-	if (AudioTracks.Num() == 0)
+	else if (Event == EMediaEvent::MediaClosed)
 	{
-		return LOCTEXT("NoAudioTracksAvailable", "No audio tracks available");
+		UrlTextBox->SetBorderBackgroundColor(FLinearColor::White);
 	}
-
-	return LOCTEXT("SelectAudioTrack", "Select an audio track...");
-}
-
-
-TSharedRef<SWidget> SMediaPlayerEditorViewer::HandleCaptionTrackComboBoxGenerateWidget(IMediaCaptionTrackPtr Value) const
-{
-	return SNew(STextBlock)
-		.Text(Value->GetStream().GetDisplayName());
 }
 
 
-void SMediaPlayerEditorViewer::HandleCaptionTrackComboBoxSelectionChanged(IMediaCaptionTrackPtr Selection, ESelectInfo::Type SelectInfo)
+FText SMediaPlayerEditorViewer::HandleOverlayCaptionText() const
 {
-	if (SelectInfo != ESelectInfo::Direct)
+	if (!GetDefault<UMediaPlayerEditorSettings>()->ShowCaptions)
 	{
-		ReloadTracks();
+		return FText::GetEmpty();
 	}
+
+	return MediaPlayer->GetCaptionText();
 }
 
 
-FText SMediaPlayerEditorViewer::HandleCaptionTrackComboBoxText() const
+FText SMediaPlayerEditorViewer::HandleOverlayPlayerNameText() const
 {
-	IMediaCaptionTrackPtr CaptionTrack = CaptionTrackComboBox->GetSelectedItem();
-
-	if (CaptionTrack.IsValid())
+	if (MediaPlayer->IsPlaying())
 	{
-		return CaptionTrack->GetStream().GetDisplayName();
-	}
-
-	if (CaptionTracks.Num() == 0)
-	{
-		return LOCTEXT("NoCaptionTracksAvailable", "No caption tracks available");
+		return FText::GetEmpty();
 	}
 
-	return LOCTEXT("SelectCaptionTrack", "Select a caption track...");
-}
-
-
-FText SMediaPlayerEditorViewer::HandleElapsedTimeTextBlockText() const
-{
-	return FText::AsTimespan(MediaPlayer->GetTime());
-}
-
-
-void SMediaPlayerEditorViewer::HandleMediaPlayerMediaChanged()
-{
-	ReloadTracks();
-}
-
+	FName PlayerName = MediaPlayer->GetPlayerName();
 
-void SMediaPlayerEditorViewer::HandleMediaPlayerTracksChanged()
-{
-	ReloadTracks();
-}
-
-
-EVisibility SMediaPlayerEditorViewer::HandleNoMediaSelectedTextVisibility() const
-{
-	return MediaPlayer->GetUrl().IsEmpty() ? EVisibility::Visible : EVisibility::Hidden;
-}
-
-
-FText SMediaPlayerEditorViewer::HandleOverlayCaptionText() const
-{
-	TSharedPtr<TArray<uint8>, ESPMode::ThreadSafe> CurrentCaption = CaptionBuffer->GetCurrentSample();
-
-	if (!CurrentCaption.IsValid())
+	if (PlayerName == NAME_None)
 	{
-		return FText::GetEmpty();
+		if (MediaPlayer->DesiredPlayerName == NAME_None)
+		{
+			return LOCTEXT("AutoPlayerName", "Auto");
+		}
+
+		return FText::FromName(MediaPlayer->DesiredPlayerName);
 	}
 
-	return FText::FromString(FString((TCHAR*)CurrentCaption->GetData()));
+	return FText::FromName(PlayerName);
 }
 
 
@@ -509,6 +487,11 @@ FText SMediaPlayerEditorViewer::HandleOverlayStateText() const
 		return LOCTEXT("StateOverlayPaused", "Paused");
 	}
 
+	if (!MediaPlayer->IsPlaying())
+	{
+		return FText::FromString(FPaths::GetBaseFilename(MediaPlayer->GetUrl()));
+	}
+
 	float Rate = MediaPlayer->GetRate();
 
 	if (FMath::IsNearlyZero(Rate) || (Rate == 1.0f))
@@ -527,27 +510,17 @@ FText SMediaPlayerEditorViewer::HandleOverlayStateText() const
 
 TOptional<float> SMediaPlayerEditorViewer::HandlePlaybackRateBoxMaxValue() const
 {
-	IMediaPlayerPtr Player = MediaPlayer->GetPlayer();
-
-	if (Player.IsValid())
-	{
-		return Player->GetMediaInfo().GetSupportedRates(EMediaPlaybackDirections::Forward, false).GetUpperBoundValue();
-	}
+	FFloatRange Rates = MediaPlayer->GetForwardRates(false);
 
-	return 0.0f;
+	return Rates.IsEmpty() ? 0.0f : Rates.GetUpperBoundValue();
 }
 
 
 TOptional<float> SMediaPlayerEditorViewer::HandlePlaybackRateBoxMinValue() const
 {
-	IMediaPlayerPtr Player = MediaPlayer->GetPlayer();
+	FFloatRange Rates = MediaPlayer->GetReverseRates(false);
 
-	if (Player.IsValid())
-	{
-		return Player->GetMediaInfo().GetSupportedRates(EMediaPlaybackDirections::Reverse, false).GetLowerBoundValue();
-	}
-
-	return 0.0f;
+	return Rates.IsEmpty() ? 0.0f : Rates.GetLowerBoundValue();
 }
 
 
@@ -563,6 +536,67 @@ void SMediaPlayerEditorViewer::HandlePlaybackRateBoxValueChanged( float NewValue
 }
 
 
+void SMediaPlayerEditorViewer::HandlePlayerMenuNewMenu(FMenuBuilder& MenuBuilder)
+{
+	// automatic player option
+	MenuBuilder.AddMenuEntry(
+		LOCTEXT("AutoPlayer", "Auto"),
+		LOCTEXT("AutoPlayerTooltip", "Select a player automatically based on the media source"),
+		FSlateIcon(),
+		FUIAction(
+			FExecuteAction::CreateLambda([this] { SetDesiredPlayerName(NAME_None); }),
+			FCanExecuteAction(),
+			FIsActionChecked::CreateLambda([this] { return MediaPlayer->DesiredPlayerName == NAME_None; })
+		),
+		NAME_None,
+		EUserInterfaceActionType::RadioButton
+	);
+
+	// get registered player plug-ins
+	auto MediaModule = FModuleManager::LoadModulePtr<IMediaModule>("Media");
+
+	if (MediaModule == nullptr)
+	{
+		return;
+	}
+
+	const TArray<IMediaPlayerFactory*>& PlayerFactories = MediaModule->GetPlayerFactories();
+
+	if (PlayerFactories.Num() == 0)
+	{
+		TSharedRef<SWidget> NoPlayersAvailableWidget = SNew(STextBlock)
+			.ColorAndOpacity(FSlateColor::UseSubduedForeground())
+			.Text(LOCTEXT("NoPlayersAvailableLabel", "No players available"));
+
+		MenuBuilder.AddWidget(NoPlayersAvailableWidget, FText::GetEmpty(), true, false);
+
+		return;
+	}
+
+	// add option for each player
+	const FString RunningPlatformName(FPlatformProperties::IniPlatformName());
+
+	for (IMediaPlayerFactory* Factory : PlayerFactories)
+	{
+		const bool SupportsRunningPlatform = Factory->GetSupportedPlatforms().Contains(RunningPlatformName);
+		const FName PlayerName = Factory->GetName();
+
+		MenuBuilder.AddMenuEntry(
+			FText::Format(LOCTEXT("PlayerNameFormat", "{0} ({1})"), Factory->GetDisplayName(), FText::FromName(PlayerName)),
+			FText::FromString(FString::Join(Factory->GetSupportedPlatforms(), TEXT(", "))),
+			FSlateIcon(),
+			FUIAction(
+				FExecuteAction::CreateLambda([=]{ SetDesiredPlayerName(PlayerName); }),
+				FCanExecuteAction::CreateLambda([=]{ return SupportsRunningPlatform; }),
+				FIsActionChecked::CreateLambda([=]{ return MediaPlayer->DesiredPlayerName == PlayerName; })
+			),
+			NAME_None,
+			EUserInterfaceActionType::RadioButton
+		);
+	}
+}
+
+
 bool SMediaPlayerEditorViewer::HandlePositionSliderIsEnabled() const
 {
 	return MediaPlayer->SupportsSeeking();
@@ -608,56 +642,123 @@ void SMediaPlayerEditorViewer::HandlePositionSliderValueChanged( float NewValue
 }
 
 
-FText SMediaPlayerEditorViewer::HandleRemainingTimeTextBlockText() const
+FReply SMediaPlayerEditorViewer::HandleTextureMouseButtonUp(const FGeometry& /*Geometry*/, const FPointerEvent& MouseEvent)
 {
-	FTimespan Remaining = MediaPlayer->GetDuration() - MediaPlayer->GetTime();
-
-	if (Remaining <= FTimespan::Zero())
+	if (MouseEvent.GetEffectingButton() == EKeys::LeftMouseButton)
 	{
-		return FText::FromString(TEXT("--:--:--"));
+		if (MediaPlayer->IsPlaying())
+		{
+			MediaPlayer->Pause();
+		}
+		else
+		{
+			MediaPlayer->Play();
+		}
+
+		return FReply::Handled();
 	}
 
-	return FText::AsTimespan(Remaining);
+	return FReply::Unhandled();
 }
 
 
-EVisibility SMediaPlayerEditorViewer::HandleTrackSelectionBoxVisibility() const
+FText SMediaPlayerEditorViewer::HandleTimerTextBlockText() const
 {
-	return MediaPlayer->GetUrl().IsEmpty() ? EVisibility::Hidden : EVisibility::Visible;
-}
+	if (!MediaPlayer->IsReady())
+	{
+		return FText::GetEmpty();
+	}
 
+	FTimespan Time = MediaPlayer->GetTime();
 
-TSharedRef<SWidget> SMediaPlayerEditorViewer::HandleVideoTrackComboBoxGenerateWidget(IMediaVideoTrackPtr Value) const
-{
-	return SNew(STextBlock)
-		.Text(Value->GetStream().GetDisplayName());
+	if (Time < FTimespan::Zero())
+	{
+		return FText::GetEmpty();
+	}
+
+	FTimespan Duration = MediaPlayer->GetDuration();
+
+	if (Duration <= FTimespan::Zero())
+	{
+		return FText::AsTimespan(Time);
+	}
+
+	return FText::Format(LOCTEXT("TimerTextBlockFormat", "{0} / {1}"), FText::AsTimespan(Time), FText::AsTimespan(Duration));
 }
 
 
-void SMediaPlayerEditorViewer::HandleVideoTrackComboBoxSelectionChanged(IMediaVideoTrackPtr Selection, ESelectInfo::Type SelectInfo)
+FText SMediaPlayerEditorViewer::HandleTimerTextBlockToolTipText() const
 {
-	if (SelectInfo != ESelectInfo::Direct)
+	if (!MediaPlayer->IsReady())
 	{
-		ReloadTracks();
+		return FText::GetEmpty();
 	}
+
+	FTimespan Remaining = MediaPlayer->GetDuration() - MediaPlayer->GetTime();
+
+	if (Remaining <= FTimespan::Zero())
+	{
+		return LOCTEXT("UnknownTimeRemainingTooltip", "Unknown time remaining");
+	}
+
+	return FText::Format(LOCTEXT("TimeRemainingTooltipFormat", "{0} remaining"), FText::AsTimespan(Remaining));
 }
 
 
-FText SMediaPlayerEditorViewer::HandleVideoTrackComboBoxText() const
+void SMediaPlayerEditorViewer::HandleTrackMenuNewMenu(FMenuBuilder& MenuBuilder, EMediaPlayerTrack TrackType)
 {
-	IMediaVideoTrackPtr VideoTrack = VideoTrackComboBox->GetSelectedItem();
+	const int32 NumTracks = MediaPlayer->GetNumTracks(TrackType);
 
-	if (VideoTrack.IsValid())
+	if (NumTracks > 0)
 	{
-		return VideoTrack->GetStream().GetDisplayName();
+		FInternationalization& I18n = FInternationalization::Get();
+
+		for (int32 TrackIndex = 0; TrackIndex < NumTracks; ++TrackIndex)
+		{
+			const FText DisplayName = MediaPlayer->GetTrackDisplayName(TrackType, TrackIndex);
+			const FString Language = MediaPlayer->GetTrackLanguage(TrackType, TrackIndex);
+			const FCulturePtr Culture = I18n.GetCulture(Language);
+			const FString LanguageName = Culture.IsValid() ? Culture->GetDisplayName() : FString();
+
+			MenuBuilder.AddMenuEntry(
+				LanguageName.IsEmpty() ? DisplayName : FText::Format(LOCTEXT("TrackNameFormat", "{0} ({1})"), DisplayName, FText::FromString(LanguageName)),
+				FText::GetEmpty(),
+				FSlateIcon(),
+				FUIAction(
+					FExecuteAction::CreateLambda([=]{ MediaPlayer->SelectTrack(TrackType, TrackIndex); }),
+					FCanExecuteAction(),
+					FIsActionChecked::CreateLambda([=]{ return (MediaPlayer->GetSelectedTrack(TrackType) == TrackIndex); })
+				),
+				NAME_None,
+				EUserInterfaceActionType::RadioButton
+			);
+		}
+
+		return;
 	}
 
-	if (VideoTracks.Num() == 0)
+	TSharedRef<SWidget> NoTracksAvailableWidget = SNew(SBox)
+		.Padding(FMargin(2.0f, 2.0f, 2.0f, 0.0f))
+		[
+			SNew(STextBlock)
+				.ColorAndOpacity(FSlateColor::UseSubduedForeground())
+				.Text(LOCTEXT("NoTracksAvailableLabel", "No tracks available"))
+		];
+
+	MenuBuilder.AddWidget(NoTracksAvailableWidget, FText::GetEmpty(), true, false);
+}
+
+
+FReply SMediaPlayerEditorViewer::HandleUrlBoxKeyDown(const FGeometry&, const FKeyEvent& KeyEvent)
+{
+	if (KeyEvent.GetKey() == EKeys::Escape)
 	{
-		return LOCTEXT("NoVideoTracksAvailable", "No video tracks available");
+		UrlTextBox->SetText(FText::FromString(MediaPlayer->GetUrl()));
+
+		return FReply::Handled().ClearUserFocus(true);
 	}
 
-	return LOCTEXT("SelectVideoTrack", "Select a video track...");
+	return FReply::Unhandled();
 }
 
 
diff --git a/Engine/Plugins/Media/MediaPlayerEditor/Source/MediaPlayerEditor/Private/Widgets/SMediaPlayerEditorViewer.h b/Engine/Plugins/Media/MediaPlayerEditor/Source/MediaPlayerEditor/Private/Widgets/SMediaPlayerEditorViewer.h
--- a/Engine/Plugins/Media/MediaPlayerEditor/Source/MediaPlayerEditor/Private/Widgets/SMediaPlayerEditorViewer.h
+++ b/Engine/Plugins/Media/MediaPlayerEditor/Source/MediaPlayerEditor/Private/Widgets/SMediaPlayerEditorViewer.h
@@ -3,6 +3,12 @@
 #pragma once
 
 
+enum class EMediaTrackType;
+class ISlateStyle;
+class SSlider;
+class UMediaPlayer;
+
+
 /**
  * Implements the contents of the viewer tab in the UMediaPlayer asset editor.
  */
@@ -31,51 +37,27 @@ public:
 	 * @param InMediaPlayer The UMediaPlayer asset to show the details for.
 	 * @param InStyleSet The style set to use.
 	 */
-	void Construct(const FArguments& InArgs, UMediaPlayer* InMediaPlayer, const TSharedRef<ISlateStyle>& InStyle);
+	void Construct(const FArguments& InArgs, UMediaPlayer& InMediaPlayer, const TSharedRef<ISlateStyle>& InStyle);
 
 protected:
 
-	/** Reloads the media texture and updates all options. */
-	void ReloadTracks();
+	/** Set the name of the desired native media player. */
+	void SetDesiredPlayerName(FName PlayerName);
 
 private:
 
 	/** Callback for the active timer. */
 	EActiveTimerReturnType HandleActiveTimer(double InCurrentTime, float InDeltaTime);
 
-	/** Callback for generating a widget for a value in the audio track combo box. */
-	TSharedRef<SWidget> HandleAudioTrackComboBoxGenerateWidget(IMediaAudioTrackPtr Value) const;
-
-	/** Callback for changing the selected audio track. */
-	void HandleAudioTrackComboBoxSelectionChanged(IMediaAudioTrackPtr Selection, ESelectInfo::Type SelectInfo);
-
-	/** Callback for getting the text in the audio track combo box. */
-	FText HandleAudioTrackComboBoxText() const;
-
-	/** Callback for generating a widget for a value in the caption track combo box. */
-	TSharedRef<SWidget> HandleCaptionTrackComboBoxGenerateWidget(IMediaCaptionTrackPtr Value) const;
-
-	/** Callback for changing the selected audio track. */
-	void HandleCaptionTrackComboBoxSelectionChanged(IMediaCaptionTrackPtr Selection, ESelectInfo::Type SelectInfo);
-
-	/** Callback for getting the text in the caption track combo box. */
-	FText HandleCaptionTrackComboBoxText() const;
-
-	/** Callback for getting the text of the 'Elapsed Time' text block. */
-	FText HandleElapsedTimeTextBlockText() const;
-
-	/** Callback for when the media player's media has changed. */
-	void HandleMediaPlayerMediaChanged();
-
-	/** Callback for when the media player's tracks changed. */
-	void HandleMediaPlayerTracksChanged();
-
-	/** Callback for getting the visibility of the 'No media selected' warning text. */
-	EVisibility HandleNoMediaSelectedTextVisibility() const;
+	/** Callback for media player events. */
+	void HandleMediaPlayerMediaEvent(EMediaEvent Event);
 
 	/** Callback for getting the text of the movie captions. */
 	FText HandleOverlayCaptionText() const;
 
+	/** Callback for getting the text of the player name overlay. */
+	FText HandleOverlayPlayerNameText() const;
+
 	/** Callback for getting the text of the playback state overlay. */
 	FText HandleOverlayStateText() const;
 
@@ -91,6 +73,9 @@ private:
 	/** Callback for committing a new value to the playback rate spin box. */
 	void HandlePlaybackRateBoxValueChanged( float NewValue );
 
+	/** Callback for creating the player sub-menu. */
+	void HandlePlayerMenuNewMenu(FMenuBuilder& MenuBuilder);
+
 	/** Callback for getting the enabled state of the position slider. */
 	bool HandlePositionSliderIsEnabled() const;
 
@@ -106,38 +91,23 @@ private:
 	/** Callback for changing the value of the 'Position' slider. */
 	void HandlePositionSliderValueChanged( float NewValue );
 
-	/** Callback for getting the text of the 'Remaining Time' text block. */
-	FText HandleRemainingTimeTextBlockText() const;
+	/** Callback for handling button up events on the movie texture border. */
+	FReply HandleTextureMouseButtonUp(const FGeometry& Geometry, const FPointerEvent& MouseEvent);
 
-	/** Callback for getting the visibility of the track selection box. */
-	EVisibility HandleTrackSelectionBoxVisibility() const;
+	/** Callback for getting the text of the timer text block. */
+	FText HandleTimerTextBlockText() const;
 
-	/** Callback for generating a widget for a value in the video track combo box. */
-	TSharedRef<SWidget> HandleVideoTrackComboBoxGenerateWidget(IMediaVideoTrackPtr Value) const;
+	/** Callback for getting the tool tip of the timer text block. */
+	FText HandleTimerTextBlockToolTipText() const;
 
-	/** Callback for changing the selected video track. */
-	void HandleVideoTrackComboBoxSelectionChanged(IMediaVideoTrackPtr Selection, ESelectInfo::Type SelectInfo);
+	/** Callback for creating the a track sub-menu. */
+	void HandleTrackMenuNewMenu(FMenuBuilder& MenuBuilder, EMediaPlayerTrack TrackType);
 
-	/** Callback for getting the text in the video track combo box. */
-	FText HandleVideoTrackComboBoxText() const;
+	/** Callback for handling key down events in the URL text box. */
+	FReply HandleUrlBoxKeyDown(const FGeometry&, const FKeyEvent& KeyEvent);
 
 private:
 
-	/** The collection of available audio tracks. */
-	TArray<IMediaAudioTrackPtr> AudioTracks;
-
-	/** Holds the audio track combo box. */
-	TSharedPtr<SComboBox<IMediaAudioTrackPtr>> AudioTrackComboBox;
-
-	/** Buffers text from the selected caption track. */
-	FMediaSampleBufferRef CaptionBuffer;
-
-	/** The collection of available caption tracks. */
-	TArray<IMediaCaptionTrackPtr> CaptionTracks;
-
-	/** Holds the caption track combo box. */
-	TSharedPtr<SComboBox<IMediaCaptionTrackPtr>> CaptionTrackComboBox;
-
 	/** Pointer to the media player that is being viewed. */
 	UMediaPlayer* MediaPlayer;
 
@@ -147,12 +117,9 @@ private:
 	/** Holds the scrubber slider. */
 	TSharedPtr<SSlider> ScrubberSlider;
 
-	/** The collection of available video tracks. */
-	TArray<IMediaVideoTrackPtr> VideoTracks;
-
-	/** Holds the video track combo box. */
-	TSharedPtr<SComboBox<IMediaVideoTrackPtr>> VideoTrackComboBox;
+	/** The style set to use for this widget. */
+	TSharedPtr<ISlateStyle> Style;
 
-	/** Holds the Slate view port. */
-	TSharedPtr<FMediaPlayerEditorViewport> Viewport;
+	/** Media URL text box. */
+	TSharedPtr<SEditableTextBox> UrlTextBox;
 };
