Copying //UE4/Dev-Sequencer to //UE4/Dev-Main (Source: //UE4/Dev-Sequencer @ 3116759)

diff --git a/Engine/Plugins/Media/AndroidMedia/Source/AndroidMediaEditor/Private/AndroidFileMediaSourceFactory.cpp b/Engine/Plugins/Media/AndroidMedia/Source/AndroidMediaEditor/Private/AndroidFileMediaSourceFactory.cpp
--- a/Engine/Plugins/Media/AndroidMedia/Source/AndroidMediaEditor/Private/AndroidFileMediaSourceFactory.cpp
+++ b/Engine/Plugins/Media/AndroidMedia/Source/AndroidMediaEditor/Private/AndroidFileMediaSourceFactory.cpp
@@ -31,7 +31,7 @@ bool UAndroidFileMediaSourceFactory::FactoryCanImport(const FString& Filename)
 UObject* UAndroidFileMediaSourceFactory::FactoryCreateFile(UClass* InClass, UObject* InParent, FName InName, EObjectFlags Flags, const FString& Filename, const TCHAR* Parms, FFeedbackContext* Warn, bool& bOutOperationCanceled)
 {
 	UFileMediaSource* MediaSource = NewObject<UFileMediaSource>(InParent, InClass, InName, Flags);
-	MediaSource->FilePath = CurrentFilename;
+	MediaSource->SetFilePath(CurrentFilename);
 
 	return MediaSource;
 }
diff --git a/Engine/Plugins/Media/AvfMedia/Source/AvfMediaEditor/Private/AvfFileMediaSourceFactory.cpp b/Engine/Plugins/Media/AvfMedia/Source/AvfMediaEditor/Private/AvfFileMediaSourceFactory.cpp
--- a/Engine/Plugins/Media/AvfMedia/Source/AvfMediaEditor/Private/AvfFileMediaSourceFactory.cpp
+++ b/Engine/Plugins/Media/AvfMedia/Source/AvfMediaEditor/Private/AvfFileMediaSourceFactory.cpp
@@ -10,9 +10,27 @@
 UAvfFileMediaSourceFactory::UAvfFileMediaSourceFactory(const FObjectInitializer& ObjectInitializer)
 	: Super(ObjectInitializer)
 {
+	Formats.Add(TEXT("3g2;3G2 Multimedia Stream"));
+	Formats.Add(TEXT("3gp;3GP Video Stream"));
+	Formats.Add(TEXT("3gp2;3GPP2 Multimedia File"));
+	Formats.Add(TEXT("3gpp;3GPP Multimedia File"));
+	Formats.Add(TEXT("ac3;AC-3 Audio File"));
+	Formats.Add(TEXT("aif;Audio Interchange File"));
+	Formats.Add(TEXT("aiff;Audio Interchange File"));
+	Formats.Add(TEXT("amr;Adaptive Multi-Rate Audio"));
+	Formats.Add(TEXT("au;Audio File"));
+	Formats.Add(TEXT("bwf;Broadcast Wave Audio"));
+	Formats.Add(TEXT("caf;Core Audio"));
+	Formats.Add(TEXT("cdda;Compact Disc Digital Audio"));
+	Formats.Add(TEXT("m4a;Apple MPEG-4 Audio"));
 	Formats.Add(TEXT("m4v;Apple MPEG-4 Video"));
 	Formats.Add(TEXT("mov;Apple QuickTime Movie"));
+	Formats.Add(TEXT("mp3;MPEG-2 Audio"));
 	Formats.Add(TEXT("mp4;MPEG-4 Movie"));
+	Formats.Add(TEXT("sdv;Samsung Digital Video"));
+	Formats.Add(TEXT("snd;Sound File"));
+	Formats.Add(TEXT("wav;Wave Audio File"));
+	Formats.Add(TEXT("wave;Wave Audio File"));
 
 	SupportedClass = UFileMediaSource::StaticClass();
 	bEditorImport = true;
@@ -24,14 +42,18 @@ UAvfFileMediaSourceFactory::UAvfFileMediaSourceFactory(const FObjectInitializer&
 
 bool UAvfFileMediaSourceFactory::FactoryCanImport(const FString& Filename)
 {
-	return true;
+	// @hack: disable file extensions that are used in other factories
+	// @todo gmp: add support for multiple factories per file extension
+	const FString FileExtension = FPaths::GetExtension(Filename);
+
+	return (FileExtension.ToUpper() != FString("WAV"));
 }
 
 
 UObject* UAvfFileMediaSourceFactory::FactoryCreateFile(UClass* InClass, UObject* InParent, FName InName, EObjectFlags Flags, const FString& Filename, const TCHAR* Parms, FFeedbackContext* Warn, bool& bOutOperationCanceled)
 {
 	UFileMediaSource* MediaSource = NewObject<UFileMediaSource>(InParent, InClass, InName, Flags);
-	MediaSource->FilePath = CurrentFilename;
+	MediaSource->SetFilePath(CurrentFilename);
 
 	return MediaSource;
 }
diff --git a/Engine/Plugins/Media/AvfMedia/Source/AvfMediaFactory/Private/AvfMediaFactoryModule.cpp b/Engine/Plugins/Media/AvfMedia/Source/AvfMediaFactory/Private/AvfMediaFactoryModule.cpp
--- a/Engine/Plugins/Media/AvfMedia/Source/AvfMediaFactory/Private/AvfMediaFactoryModule.cpp
+++ b/Engine/Plugins/Media/AvfMedia/Source/AvfMediaFactory/Private/AvfMediaFactoryModule.cpp
@@ -102,9 +102,29 @@ public:
 		SupportedPlatforms.Add(TEXT("Mac"));
 
 		// supported file extensions
+		SupportedFileExtensions.Add(TEXT("3g2"));
+		SupportedFileExtensions.Add(TEXT("3gp"));
+		SupportedFileExtensions.Add(TEXT("3gp2"));
+		SupportedFileExtensions.Add(TEXT("3gpp"));
+		SupportedFileExtensions.Add(TEXT("ac3"));
+		SupportedFileExtensions.Add(TEXT("aif"));
+		SupportedFileExtensions.Add(TEXT("aiff"));
+		SupportedFileExtensions.Add(TEXT("aifc"));
+		SupportedFileExtensions.Add(TEXT("amr"));
+		SupportedFileExtensions.Add(TEXT("au"));
+		SupportedFileExtensions.Add(TEXT("bwf"));
+		SupportedFileExtensions.Add(TEXT("caf"));
+		SupportedFileExtensions.Add(TEXT("cdda"));
+		SupportedFileExtensions.Add(TEXT("m4a"));
 		SupportedFileExtensions.Add(TEXT("m4v"));
 		SupportedFileExtensions.Add(TEXT("mov"));
+		SupportedFileExtensions.Add(TEXT("mp3"));
 		SupportedFileExtensions.Add(TEXT("mp4"));
+		SupportedFileExtensions.Add(TEXT("qt"));
+		SupportedFileExtensions.Add(TEXT("sdv"));
+		SupportedFileExtensions.Add(TEXT("snd"));
+		SupportedFileExtensions.Add(TEXT("wav"));
+		SupportedFileExtensions.Add(TEXT("wave"));
 
 		// register player factory
 		auto MediaModule = FModuleManager::LoadModulePtr<IMediaModule>("Media");
diff --git a/Engine/Plugins/Media/MediaPlayerEditor/Source/MediaPlayerEditor/Private/AssetTools/FileMediaSourceActions.cpp b/Engine/Plugins/Media/MediaPlayerEditor/Source/MediaPlayerEditor/Private/AssetTools/FileMediaSourceActions.cpp
--- a/Engine/Plugins/Media/MediaPlayerEditor/Source/MediaPlayerEditor/Private/AssetTools/FileMediaSourceActions.cpp
+++ b/Engine/Plugins/Media/MediaPlayerEditor/Source/MediaPlayerEditor/Private/AssetTools/FileMediaSourceActions.cpp
@@ -18,6 +18,12 @@ FFileMediaSourceActions::FFileMediaSourceActions(const TSharedRef<ISlateStyle>&
 /* FAssetTypeActions_Base interface
  *****************************************************************************/
 
+bool FFileMediaSourceActions::CanFilter()
+{
+	return true;
+}
+
+
 FText FFileMediaSourceActions::GetName() const
 {
 	return LOCTEXT("AssetTypeActions_FileMediaSource", "File Media Source");
diff --git a/Engine/Plugins/Media/MediaPlayerEditor/Source/MediaPlayerEditor/Private/AssetTools/FileMediaSourceActions.h b/Engine/Plugins/Media/MediaPlayerEditor/Source/MediaPlayerEditor/Private/AssetTools/FileMediaSourceActions.h
--- a/Engine/Plugins/Media/MediaPlayerEditor/Source/MediaPlayerEditor/Private/AssetTools/FileMediaSourceActions.h
+++ b/Engine/Plugins/Media/MediaPlayerEditor/Source/MediaPlayerEditor/Private/AssetTools/FileMediaSourceActions.h
@@ -24,6 +24,7 @@ public:
 
 	//~ FAssetTypeActions_Base interface
 
+	virtual bool CanFilter() override;
 	virtual FText GetName() const override;
 	virtual UClass* GetSupportedClass() const override;
 	virtual TSharedPtr<class SWidget> GetThumbnailOverlay(const FAssetData& AssetData) const override;
diff --git a/Engine/Plugins/Media/MediaPlayerEditor/Source/MediaPlayerEditor/Private/AssetTools/MediaSourceActions.cpp b/Engine/Plugins/Media/MediaPlayerEditor/Source/MediaPlayerEditor/Private/AssetTools/MediaSourceActions.cpp
--- a/Engine/Plugins/Media/MediaPlayerEditor/Source/MediaPlayerEditor/Private/AssetTools/MediaSourceActions.cpp
+++ b/Engine/Plugins/Media/MediaPlayerEditor/Source/MediaPlayerEditor/Private/AssetTools/MediaSourceActions.cpp
@@ -10,6 +10,12 @@
 /* FAssetTypeActions_Base interface
  *****************************************************************************/
 
+bool FMediaSourceActions::CanFilter()
+{
+	return false;
+}
+
+
 FText FMediaSourceActions::GetAssetDescription(const class FAssetData& AssetData) const
 {
 	auto MediaSource = Cast<UMediaSource>(AssetData.GetAsset());
diff --git a/Engine/Plugins/Media/MediaPlayerEditor/Source/MediaPlayerEditor/Private/Customizations/FileMediaSourceCustomization.cpp b/Engine/Plugins/Media/MediaPlayerEditor/Source/MediaPlayerEditor/Private/Customizations/FileMediaSourceCustomization.cpp
--- a/Engine/Plugins/Media/MediaPlayerEditor/Source/MediaPlayerEditor/Private/Customizations/FileMediaSourceCustomization.cpp
+++ b/Engine/Plugins/Media/MediaPlayerEditor/Source/MediaPlayerEditor/Private/Customizations/FileMediaSourceCustomization.cpp
@@ -119,7 +119,7 @@ FString FFileMediaSourceCustomization::HandleFilePathPickerFileTypeFilter() cons
 
 void FFileMediaSourceCustomization::HandleFilePathPickerPathPicked(const FString& PickedPath)
 {
-	if (PickedPath.IsEmpty() || PickedPath.StartsWith(TEXT("./")) || PickedPath.Contains(TEXT("://")))
+	if (PickedPath.IsEmpty() || PickedPath.StartsWith(TEXT("./")))
 	{
 		FilePathProperty->SetValue(PickedPath);
 	}
diff --git a/Engine/Plugins/Media/MediaPlayerEditor/Source/MediaPlayerEditor/Private/Shared/MediaPlayerEditorStyle.h b/Engine/Plugins/Media/MediaPlayerEditor/Source/MediaPlayerEditor/Private/Shared/MediaPlayerEditorStyle.h
--- a/Engine/Plugins/Media/MediaPlayerEditor/Source/MediaPlayerEditor/Private/Shared/MediaPlayerEditorStyle.h
+++ b/Engine/Plugins/Media/MediaPlayerEditor/Source/MediaPlayerEditor/Private/Shared/MediaPlayerEditorStyle.h
@@ -36,6 +36,7 @@ public:
 		Set("MediaPlayerEditor.SettingsButton", new IMAGE_BRUSH("btn_settings_16x", Icon12x12));
 
 		// misc
+		Set("MediaPlayerEditor.DragDropBorder", new BOX_BRUSH("border_dragdrop.png", 0.5f));
 		Set("MediaPlayerEditor.FileMediaSourcePrecached", new IMAGE_BRUSH("filemediasource_precached", Icon12x12));
 		Set("MediaPlayerEditor.MediaSourceOpened", new IMAGE_BRUSH("mediasource_opened", Icon8x8));
 
diff --git a/Engine/Plugins/Media/MediaPlayerEditor/Source/MediaPlayerEditor/Private/Widgets/SMediaPlayerEditorMedia.cpp b/Engine/Plugins/Media/MediaPlayerEditor/Source/MediaPlayerEditor/Private/Widgets/SMediaPlayerEditorMedia.cpp
--- a/Engine/Plugins/Media/MediaPlayerEditor/Source/MediaPlayerEditor/Private/Widgets/SMediaPlayerEditorMedia.cpp
+++ b/Engine/Plugins/Media/MediaPlayerEditor/Source/MediaPlayerEditor/Private/Widgets/SMediaPlayerEditorMedia.cpp
@@ -180,7 +180,7 @@ TSharedPtr<SWidget> SMediaPlayerEditorMedia::HandleAssetPickerGetAssetContextMen
 					FSlateIcon(FEditorStyle::GetStyleSetName(), "SystemWideCommands.FindInContentBrowser"),
 					FUIAction(
 						FExecuteAction::CreateLambda([=]() {
-							FPlatformProcess::ExploreFolder(*FileMediaSource->FilePath);
+							FPlatformProcess::ExploreFolder(*FileMediaSource->GetFilePath());
 						}),
 						FCanExecuteAction::CreateLambda([=]() -> bool {
 							return FileMediaSource->Validate();
diff --git a/Engine/Plugins/Media/MediaPlayerEditor/Source/MediaPlayerEditor/Private/Widgets/SMediaPlayerEditorViewer.cpp b/Engine/Plugins/Media/MediaPlayerEditor/Source/MediaPlayerEditor/Private/Widgets/SMediaPlayerEditorViewer.cpp
--- a/Engine/Plugins/Media/MediaPlayerEditor/Source/MediaPlayerEditor/Private/Widgets/SMediaPlayerEditorViewer.cpp
+++ b/Engine/Plugins/Media/MediaPlayerEditor/Source/MediaPlayerEditor/Private/Widgets/SMediaPlayerEditorViewer.cpp
@@ -13,7 +13,9 @@
  *****************************************************************************/
 
 SMediaPlayerEditorViewer::SMediaPlayerEditorViewer()
-	: MediaPlayer(nullptr)
+	: DragOver(false)
+	, DragValid(false)
+	, MediaPlayer(nullptr)
 { }
 
 
@@ -142,234 +144,320 @@ void SMediaPlayerEditorViewer::Construct(const FArguments& InArgs, UMediaPlayer&
 
 	ChildSlot
 	[
-		SNew(SVerticalBox)
+		SNew(SOverlay)
 
-		+ SVerticalBox::Slot()
-			.AutoHeight()
+		+ SOverlay::Slot()
 			[
-				SNew(SHorizontalBox)
+				SNew(SVerticalBox)
 
-				+ SHorizontalBox::Slot()
-					.FillWidth(1.0f)
+				+ SVerticalBox::Slot()
+					.AutoHeight()
 					[
-						// url box
-						SAssignNew(UrlTextBox, SEditableTextBox)
-							.ClearKeyboardFocusOnCommit(true)
-							.HintText(LOCTEXT("UrlTextBoxHint", "Media URL"))
-							.ToolTipText(LOCTEXT("UrlTextBoxToolTip", "Enter the URL of a media source"))
-							.OnKeyDownHandler(this, &SMediaPlayerEditorViewer::HandleUrlBoxKeyDown)
-							.OnTextCommitted_Lambda([this](const FText& InText, ETextCommit::Type InCommitType){
-								if (InCommitType == ETextCommit::OnEnter)
-								{
-									OpenUrlTextBoxUrl();
-								}
-							})
-					]
+						SNew(SHorizontalBox)
 
-				+ SHorizontalBox::Slot()
-					.AutoWidth()
-					.Padding(4.0f, 0.0f, 0.0f, 0.0f)
-					[
-						// go button
-						SNew(SButton)
-							.ToolTipText_Lambda([this]() {
-								return ((UrlTextBox->GetText().ToString() == MediaPlayer->GetUrl()) && !MediaPlayer->GetUrl().IsEmpty())
-									? LOCTEXT("ReloadButtonToolTip", "Reload the current media URL")
-									: LOCTEXT("GoButtonToolTip", "Open the specified media URL");
-							})
-							.IsEnabled_Lambda([this]{
-								return !UrlTextBox->GetText().IsEmpty();
-							})
-							.OnClicked_Lambda([this]{
-								OpenUrlTextBoxUrl();
-								return FReply::Handled();
-							})
+						+ SHorizontalBox::Slot()
+							.FillWidth(1.0f)
 							[
-								SNew(SImage)
-									.Image_Lambda([this]() {
-										return ((UrlTextBox->GetText().ToString() == MediaPlayer->GetUrl()) && !MediaPlayer->GetUrl().IsEmpty())
-											? Style->GetBrush("MediaPlayerEditor.ReloadButton")
-											: Style->GetBrush("MediaPlayerEditor.GoButton");
+								// url box
+								SAssignNew(UrlTextBox, SEditableTextBox)
+									.ClearKeyboardFocusOnCommit(true)
+									.HintText(LOCTEXT("UrlTextBoxHint", "Media URL"))
+									.ToolTipText(LOCTEXT("UrlTextBoxToolTip", "Enter the URL of a media source"))
+									.OnKeyDownHandler(this, &SMediaPlayerEditorViewer::HandleUrlBoxKeyDown)
+									.OnTextCommitted_Lambda([this](const FText& InText, ETextCommit::Type InCommitType){
+										if (InCommitType == ETextCommit::OnEnter)
+										{
+											OpenUrlTextBoxUrl();
+										}
 									})
 							]
-					]
-			]
-
-		+ SVerticalBox::Slot()
-			.FillHeight(1.0f)
-			.Padding(0.0f, 4.0f, 0.0f, 0.0f)
-			[
-				SNew(SBorder)
-					.BorderImage(FCoreStyle::Get().GetBrush("BlackBrush"))
-					.Padding(0.0f)
-					.OnMouseButtonUp(this, &SMediaPlayerEditorViewer::HandleTextureMouseButtonUp)
-					[
-						SNew(SOverlay)
 
-						+ SOverlay::Slot()
+						+ SHorizontalBox::Slot()
+							.AutoWidth()
+							.Padding(4.0f, 0.0f, 0.0f, 0.0f)
 							[
-								// movie texture
-								SNew(SScaleBox)
-									.Stretch_Lambda([]() -> EStretch::Type { return GetDefault<UMediaPlayerEditorSettings>()->ScaleVideoToFit ? EStretch::ScaleToFit : EStretch::Fill; })
+								// go button
+								SNew(SButton)
+									.ToolTipText_Lambda([this]() {
+										return ((UrlTextBox->GetText().ToString() == MediaPlayer->GetUrl()) && !MediaPlayer->GetUrl().IsEmpty())
+											? LOCTEXT("ReloadButtonToolTip", "Reload the current media URL")
+											: LOCTEXT("GoButtonToolTip", "Open the specified media URL");
+									})
+									.IsEnabled_Lambda([this]{
+										return !UrlTextBox->GetText().IsEmpty();
+									})
+									.OnClicked_Lambda([this]{
+										OpenUrlTextBoxUrl();
+										return FReply::Handled();
+									})
 									[
-										SNew(SMediaPlayerEditorOutput, InMediaPlayer)
+										SNew(SImage)
+											.Image_Lambda([this]() {
+												return ((UrlTextBox->GetText().ToString() == MediaPlayer->GetUrl()) && !MediaPlayer->GetUrl().IsEmpty())
+													? Style->GetBrush("MediaPlayerEditor.ReloadButton")
+													: Style->GetBrush("MediaPlayerEditor.GoButton");
+											})
 									]
 							]
+					]
 
-						+ SOverlay::Slot()
-							.Padding(FMargin(12.0f, 8.0f))
+				+ SVerticalBox::Slot()
+					.FillHeight(1.0f)
+					.Padding(0.0f, 4.0f, 0.0f, 0.0f)
+					[
+						SNew(SBorder)
+							.BorderImage(FCoreStyle::Get().GetBrush("BlackBrush"))
+							.Padding(0.0f)
+							.OnMouseButtonUp(this, &SMediaPlayerEditorViewer::HandleTextureMouseButtonUp)
 							[
-								SNew(SHorizontalBox)
+								SNew(SOverlay)
 
-								+ SHorizontalBox::Slot()
-									.HAlign(HAlign_Left)
-									.VAlign(VAlign_Top)
+								+ SOverlay::Slot()
 									[
-										// playback state
-										SNew(STextBlock)
-											.ColorAndOpacity(FSlateColor::UseSubduedForeground())
-											.Font(FSlateFontInfo(FPaths::EngineContentDir() / TEXT("Slate/Fonts/Roboto-Regular.ttf"), 18))
-											.ShadowOffset(FVector2D(1.f, 1.f))
-											.Text(this, &SMediaPlayerEditorViewer::HandleOverlayStateText)
+										// movie texture
+										SNew(SScaleBox)
+											.Stretch_Lambda([]() -> EStretch::Type { return GetDefault<UMediaPlayerEditorSettings>()->ScaleVideoToFit ? EStretch::ScaleToFit : EStretch::Fill; })
+											[
+												SNew(SMediaPlayerEditorOutput, InMediaPlayer)
+											]
 									]
 
-								+ SHorizontalBox::Slot()
-									.HAlign(HAlign_Right)
-									.VAlign(VAlign_Top)
+								+ SOverlay::Slot()
+									.Padding(FMargin(12.0f, 8.0f))
 									[
-										// player name
-										SNew(STextBlock)
-											.ColorAndOpacity(FSlateColor::UseSubduedForeground())
-											.Font(FSlateFontInfo(FPaths::EngineContentDir() / TEXT("Slate/Fonts/Roboto-Regular.ttf"), 18))
-											.ShadowOffset(FVector2D(1.f, 1.f))
-											.Text(this, &SMediaPlayerEditorViewer::HandleOverlayPlayerNameText)
-									]
-							]
+										SNew(SHorizontalBox)
 
-						+ SOverlay::Slot()
-							.HAlign(HAlign_Fill)
-							.VAlign(VAlign_Bottom)
-							.Padding(FMargin(50.0f, 20.0f))
-							[
-								// caption text
-								SNew(STextBlock)
-									.AutoWrapText(true)
-									.Font(FSlateFontInfo(FPaths::EngineContentDir() / TEXT("Slate/Fonts/Roboto-Regular.ttf"), 14))
-									.ShadowOffset(FVector2D(1.f, 1.f))
-									.Text(this, &SMediaPlayerEditorViewer::HandleOverlayCaptionText)
-							]
-					]
-			]
-
-		+ SVerticalBox::Slot()
-			.AutoHeight()
-			.Padding(0.0f, 2.0f, 0.0f, 0.0f)
-			.VAlign(VAlign_Center)
-			[
-				SNew(SBorder)
-					.BorderImage(FEditorStyle::GetBrush("ToolPanel.GroupBorder"))
-					.ForegroundColor(FLinearColor::Gray)
-					.Padding(6.0f)
-					[
-						SNew(SVerticalBox)
-
-						+ SVerticalBox::Slot()
-							.AutoHeight()
-							[
-								SNew(SOverlay)
+										+ SHorizontalBox::Slot()
+											.HAlign(HAlign_Left)
+											.VAlign(VAlign_Top)
+											[
+												// playback state
+												SNew(STextBlock)
+													.ColorAndOpacity(FSlateColor::UseSubduedForeground())
+													.Font(FSlateFontInfo(FPaths::EngineContentDir() / TEXT("Slate/Fonts/Roboto-Regular.ttf"), 18))
+													.ShadowOffset(FVector2D(1.f, 1.f))
+													.Text(this, &SMediaPlayerEditorViewer::HandleOverlayStateText)
+											]
 
-								+ SOverlay::Slot()
-									[
-										// scrubber
-										SAssignNew(ScrubberSlider, SSlider)
-											.IsEnabled(this, &SMediaPlayerEditorViewer::HandlePositionSliderIsEnabled)
-											.OnMouseCaptureBegin(this, &SMediaPlayerEditorViewer::HandlePositionSliderMouseCaptureBegin)
-											.OnMouseCaptureEnd(this, &SMediaPlayerEditorViewer::HandlePositionSliderMouseCaptureEnd)
-											.OnValueChanged(this, &SMediaPlayerEditorViewer::HandlePositionSliderValueChanged)
-											.Orientation(Orient_Horizontal)
-											.Value(this, &SMediaPlayerEditorViewer::HandlePositionSliderValue)
+										+ SHorizontalBox::Slot()
+											.HAlign(HAlign_Right)
+											.VAlign(VAlign_Top)
+											[
+												// player name
+												SNew(STextBlock)
+													.ColorAndOpacity(FSlateColor::UseSubduedForeground())
+													.Font(FSlateFontInfo(FPaths::EngineContentDir() / TEXT("Slate/Fonts/Roboto-Regular.ttf"), 18))
+													.ShadowOffset(FVector2D(1.f, 1.f))
+													.Text(this, &SMediaPlayerEditorViewer::HandleOverlayPlayerNameText)
+											]
 									]
 
 								+ SOverlay::Slot()
+									.HAlign(HAlign_Fill)
+									.VAlign(VAlign_Bottom)
+									.Padding(FMargin(50.0f, 20.0f))
 									[
-										SNew(SProgressBar)
-											.ToolTipText(LOCTEXT("BufferingTooltip", "Buffering..."))
-											.Visibility_Lambda([this]() -> EVisibility { return MediaPlayer->IsPreparing() ? EVisibility::Visible : EVisibility::Hidden; })
+										// caption text
+										SNew(STextBlock)
+											.AutoWrapText(true)
+											.Font(FSlateFontInfo(FPaths::EngineContentDir() / TEXT("Slate/Fonts/Roboto-Regular.ttf"), 14))
+											.ShadowOffset(FVector2D(1.f, 1.f))
+											.Text(this, &SMediaPlayerEditorViewer::HandleOverlayCaptionText)
 									]
 							]
+					]
 
-						+ SVerticalBox::Slot()
-							.AutoHeight()
-							.Padding(0.0f, 2.0f, 0.0f, 0.0f)
+				+ SVerticalBox::Slot()
+					.AutoHeight()
+					.Padding(0.0f, 2.0f, 0.0f, 0.0f)
+					.VAlign(VAlign_Center)
+					[
+						SNew(SBorder)
+							.BorderImage(FEditorStyle::GetBrush("ToolPanel.GroupBorder"))
+							.ForegroundColor(FLinearColor::Gray)
+							.Padding(6.0f)
 							[
-								SNew(SHorizontalBox)
+								SNew(SVerticalBox)
 
-								// timer
-								+ SHorizontalBox::Slot()
-									.AutoWidth()
-									.Padding(0.0f, 0.0f, 4.0f, 0.0f)
-									.VAlign(VAlign_Center)
+								+ SVerticalBox::Slot()
+									.AutoHeight()
 									[
-										SNew(STextBlock)
-											.Text(this, &SMediaPlayerEditorViewer::HandleTimerTextBlockText)
-											.ToolTipText(this, &SMediaPlayerEditorViewer::HandleTimerTextBlockToolTipText)
-									]
+										SNew(SOverlay)
 
-								// transport controls
-								+ SHorizontalBox::Slot()
-									.FillWidth(1.0f)
-									.Padding(8.0f, 0.0f)
-									.VAlign(VAlign_Center)
-									[
-										SNew(SSpacer)
-										//SNew(STransportControl)
+										+ SOverlay::Slot()
+											[
+												// scrubber
+												SAssignNew(ScrubberSlider, SSlider)
+													.IsEnabled(this, &SMediaPlayerEditorViewer::HandlePositionSliderIsEnabled)
+													.OnMouseCaptureBegin(this, &SMediaPlayerEditorViewer::HandlePositionSliderMouseCaptureBegin)
+													.OnMouseCaptureEnd(this, &SMediaPlayerEditorViewer::HandlePositionSliderMouseCaptureEnd)
+													.OnValueChanged(this, &SMediaPlayerEditorViewer::HandlePositionSliderValueChanged)
+													.Orientation(Orient_Horizontal)
+													.Value(this, &SMediaPlayerEditorViewer::HandlePositionSliderValue)
+											]
+
+										+ SOverlay::Slot()
+											[
+												SNew(SProgressBar)
+													.ToolTipText(LOCTEXT("BufferingTooltip", "Buffering..."))
+													.Visibility_Lambda([this]() -> EVisibility { return MediaPlayer->IsPreparing() ? EVisibility::Visible : EVisibility::Hidden; })
+											]
 									]
 
-								// settings
-								+ SHorizontalBox::Slot()
-									.AutoWidth()
-									.Padding(4.0f, 0.0f, 0.0f, 0.0f)
-									.VAlign(VAlign_Center)
+								+ SVerticalBox::Slot()
+									.AutoHeight()
+									.Padding(0.0f, 2.0f, 0.0f, 0.0f)
 									[
-										SNew(SComboButton)
-											.ContentPadding(0.0f)
-											.ButtonContent()
+										SNew(SHorizontalBox)
+
+										// timer
+										+ SHorizontalBox::Slot()
+											.AutoWidth()
+											.Padding(0.0f, 0.0f, 4.0f, 0.0f)
+											.VAlign(VAlign_Center)
 											[
-												SNew(SHorizontalBox)
+												SNew(STextBlock)
+													.Text(this, &SMediaPlayerEditorViewer::HandleTimerTextBlockText)
+													.ToolTipText(this, &SMediaPlayerEditorViewer::HandleTimerTextBlockToolTipText)
+											]
+
+										// transport controls
+										+ SHorizontalBox::Slot()
+											.FillWidth(1.0f)
+											.Padding(8.0f, 0.0f)
+											.VAlign(VAlign_Center)
+											[
+												SNew(SSpacer)
+												//SNew(STransportControl)
+											]
 
-												+ SHorizontalBox::Slot()
-													.AutoWidth()
-													.VAlign(VAlign_Center)
+										// settings
+										+ SHorizontalBox::Slot()
+											.AutoWidth()
+											.Padding(4.0f, 0.0f, 0.0f, 0.0f)
+											.VAlign(VAlign_Center)
+											[
+												SNew(SComboButton)
+													.ContentPadding(0.0f)
+													.ButtonContent()
 													[
-														SNew(SImage)
-															.Image(InStyle->GetBrush("MediaPlayerEditor.SettingsButton"))
+														SNew(SHorizontalBox)
+
+														+ SHorizontalBox::Slot()
+															.AutoWidth()
+															.VAlign(VAlign_Center)
+															[
+																SNew(SImage)
+																	.Image(InStyle->GetBrush("MediaPlayerEditor.SettingsButton"))
+															]
+
+														+ SHorizontalBox::Slot()
+															.AutoWidth()
+															.Padding(3.0f, 0.0f, 0.0f, 0.0f)
+															.VAlign(VAlign_Center)
+															[
+																SNew(STextBlock)
+																	.Text(LOCTEXT("OptionsButton", "Playback Options"))
+															]
 													]
-
-												+ SHorizontalBox::Slot()
-													.AutoWidth()
-													.Padding(3.0f, 0.0f, 0.0f, 0.0f)
-													.VAlign(VAlign_Center)
+													.ButtonStyle(FEditorStyle::Get(), "ToggleButton")
+													.ForegroundColor(FSlateColor::UseForeground())
+													.MenuContent()
 													[
-														SNew(STextBlock)
-															.Text(LOCTEXT("OptionsButton", "Playback Options"))
+														SettingsMenuBuilder.MakeWidget()
 													]
 											]
-											.ButtonStyle(FEditorStyle::Get(), "ToggleButton")
-											.ForegroundColor(FSlateColor::UseForeground())
-											.MenuContent()
-											[
-												SettingsMenuBuilder.MakeWidget()
-											]
 									]
 							]
 					]
 			]
+
+		+ SOverlay::Slot()
+			[
+				SNew(SBorder)
+					.BorderImage(FEditorStyle::GetBrush("WhiteBrush"))
+					.BorderBackgroundColor_Lambda([this]() -> FLinearColor {
+						return DragValid
+							? FLinearColor(0.0f, 1.0f, 0.0f, 0.15f)
+							: FLinearColor(1.0f, 0.0f, 0.0f, 0.15f);
+					})
+					.Visibility_Lambda([this]() -> EVisibility {
+						return DragOver && FSlateApplication::Get().IsDragDropping()
+							? EVisibility::HitTestInvisible
+							: EVisibility::Hidden;
+					})
+			]
 	];
 
 	RegisterActiveTimer(0.0f, FWidgetActiveTimerDelegate::CreateSP(this, &SMediaPlayerEditorViewer::HandleActiveTimer));
 }
 
 
+/* SWidget interface
+ *****************************************************************************/
+
+void SMediaPlayerEditorViewer::OnDragEnter(const FGeometry& MyGeometry, const FDragDropEvent& DragDropEvent)
+{
+	DragOver = true;
+
+	TSharedPtr<FExternalDragOperation> DragDropOp = DragDropEvent.GetOperationAs<FExternalDragOperation>();
+	
+	if (DragDropOp.IsValid() && DragDropOp->HasFiles())
+	{
+		const TArray<FString>& Files = DragDropOp->GetFiles();
+		
+		if (Files.Num() == 1)
+		{
+			DragValid = MediaPlayer->CanPlayUrl(FString(TEXT("file://")) + Files[0]);
+
+			return;
+		}
+	}
+
+	DragValid = false;
+}
+
+
+void SMediaPlayerEditorViewer::OnDragLeave(const FDragDropEvent& DragDropEvent)
+{
+	DragOver = false;
+}
+
+
+FReply SMediaPlayerEditorViewer::OnDragOver(const FGeometry& MyGeometry, const FDragDropEvent& DragDropEvent)
+{
+	if (DragValid)
+	{
+		return FReply::Handled();
+	}
+
+	return SCompoundWidget::OnDragOver(MyGeometry,DragDropEvent);
+}
+
+
+FReply SMediaPlayerEditorViewer::OnDrop(const FGeometry& MyGeometry, const FDragDropEvent& DragDropEvent)
+{
+	if (DragValid)
+	{
+		TSharedPtr<FExternalDragOperation> DragDropOp = DragDropEvent.GetOperationAs<FExternalDragOperation>();
+
+		if (DragDropOp.IsValid() && DragDropOp->HasFiles())
+		{
+			const TArray<FString>& Files = DragDropOp->GetFiles();
+
+			if (Files.Num() == 1)
+			{
+				MediaPlayer->OpenUrl(FString(TEXT("file://")) + Files[0]);
+
+				return FReply::Handled();
+			}
+		}
+	}
+
+	return SCompoundWidget::OnDrop(MyGeometry,DragDropEvent);
+}
+
+
 /* SMediaPlayerEditorPlayer implementation
  *****************************************************************************/
 
diff --git a/Engine/Plugins/Media/WmfMedia/Source/WmfMedia/Private/Player/WmfMediaPlayer.cpp b/Engine/Plugins/Media/WmfMedia/Source/WmfMedia/Private/Player/WmfMediaPlayer.cpp
--- a/Engine/Plugins/Media/WmfMedia/Source/WmfMedia/Private/Player/WmfMediaPlayer.cpp
+++ b/Engine/Plugins/Media/WmfMedia/Source/WmfMedia/Private/Player/WmfMediaPlayer.cpp
@@ -505,8 +505,9 @@ void FWmfMediaPlayer::AddStreamToTopology(uint32 StreamIndex, IMFTopology* Topol
 			const bool OutputRgba = 
 				(SubType != MFVideoFormat_H264) &&		// H.264 to RGB is broken
 				(SubType != MFVideoFormat_H264_ES) &&	// H.264 to RGB is broken
-				(SubType != MFVideoFormat_M4S2) &&		// incorrect brightness
-				(SubType != MFVideoFormat_WMV3);		// Kite demo outro issue (UE-35162)
+				(SubType != MFVideoFormat_M4S2) &&		// YUV to RGB has incorrect brightness
+				(SubType != MFVideoFormat_WMV2) &&		// YUV to RGB has incorrect brightness
+				(SubType != MFVideoFormat_WMV3);		// brightness issues (UE-35162)
 
 			Result = OutputType->SetGUID(MF_MT_SUBTYPE, OutputRgba ? MFVideoFormat_RGB32 : MFVideoFormat_YUY2);
 
diff --git a/Engine/Plugins/Media/WmfMedia/Source/WmfMediaEditor/Private/WmfFileMediaSourceFactory.cpp b/Engine/Plugins/Media/WmfMedia/Source/WmfMediaEditor/Private/WmfFileMediaSourceFactory.cpp
--- a/Engine/Plugins/Media/WmfMedia/Source/WmfMediaEditor/Private/WmfFileMediaSourceFactory.cpp
+++ b/Engine/Plugins/Media/WmfMedia/Source/WmfMediaEditor/Private/WmfFileMediaSourceFactory.cpp
@@ -39,19 +39,18 @@ UWmfFileMediaSourceFactory::UWmfFileMediaSourceFactory(const FObjectInitializer&
 
 bool UWmfFileMediaSourceFactory::FactoryCanImport(const FString& Filename)
 {
-	FString FileExtension = FPaths::GetExtension(Filename);
-	if (FileExtension.ToUpper() == FString("WAV"))
-	{
-		return false;
-	}
-	return true;
+	// @hack: disable file extensions that are used in other factories
+	// @todo gmp: add support for multiple factories per file extension
+	const FString FileExtension = FPaths::GetExtension(Filename);
+
+	return (FileExtension.ToUpper() != FString("WAV"));
 }
 
 
 UObject* UWmfFileMediaSourceFactory::FactoryCreateFile(UClass* InClass, UObject* InParent, FName InName, EObjectFlags Flags, const FString& Filename, const TCHAR* Parms, FFeedbackContext* Warn, bool& bOutOperationCanceled)
 {
 	UFileMediaSource* MediaSource = NewObject<UFileMediaSource>(InParent, InClass, InName, Flags);
-	MediaSource->FilePath = CurrentFilename;
+	MediaSource->SetFilePath(CurrentFilename);
 
 	return MediaSource;
 }
diff --git a/Engine/Plugins/MovieScene/LevelSequenceEditor/Source/LevelSequenceEditor/Private/LevelSequenceEditorToolkit.cpp b/Engine/Plugins/MovieScene/LevelSequenceEditor/Source/LevelSequenceEditor/Private/LevelSequenceEditorToolkit.cpp
--- a/Engine/Plugins/MovieScene/LevelSequenceEditor/Source/LevelSequenceEditor/Private/LevelSequenceEditorToolkit.cpp
+++ b/Engine/Plugins/MovieScene/LevelSequenceEditor/Source/LevelSequenceEditor/Private/LevelSequenceEditorToolkit.cpp
@@ -401,7 +401,7 @@ void FLevelSequenceEditorToolkit::AddDefaultTracksForActor(AActor& Actor, const
 					}
 
 					// @todo sequencer: hack: setting defaults for transform tracks
-					if (NewTrack->IsA(UMovieScene3DTransformTrack::StaticClass()))
+					if (NewTrack->IsA(UMovieScene3DTransformTrack::StaticClass()) && Sequencer->GetAutoSetTrackDefaults())
 					{
 						auto TransformSection = Cast<UMovieScene3DTransformSection>(NewSection);
 
@@ -545,14 +545,7 @@ void FLevelSequenceEditorToolkit::AddDefaultTracksForActor(AActor& Actor, const
 			}
 
 			// key property
-			FKeyPropertyParams KeyPropertyParams(TArrayBuilder<UObject*>().Add(PropertyOwner), PropertyPath);
-			{
-				KeyPropertyParams.KeyParams.bCreateTrackIfMissing = true;
-				KeyPropertyParams.KeyParams.bCreateHandleIfMissing = true;
-				KeyPropertyParams.KeyParams.bCreateKeyIfUnchanged = false;
-				KeyPropertyParams.KeyParams.bCreateKeyIfEmpty = false;
-				KeyPropertyParams.KeyParams.bCreateKeyOnlyWhenAutoKeying = false;
-			}
+			FKeyPropertyParams KeyPropertyParams(TArrayBuilder<UObject*>().Add(PropertyOwner), PropertyPath, ESequencerKeyMode::ManualKey);
 
 			Sequencer->KeyProperty(KeyPropertyParams);
 
diff --git a/Engine/Plugins/MovieScene/LevelSequenceEditor/Source/LevelSequenceEditor/Private/Misc/LevelSequenceEditorHelpers.cpp b/Engine/Plugins/MovieScene/LevelSequenceEditor/Source/LevelSequenceEditor/Private/Misc/LevelSequenceEditorHelpers.cpp
--- a/Engine/Plugins/MovieScene/LevelSequenceEditor/Source/LevelSequenceEditor/Private/Misc/LevelSequenceEditorHelpers.cpp
+++ b/Engine/Plugins/MovieScene/LevelSequenceEditor/Source/LevelSequenceEditor/Private/Misc/LevelSequenceEditorHelpers.cpp
@@ -60,15 +60,18 @@ class SMasterSequenceSettings : public SCompoundWidget, public FGCObject
 			SNew(SVerticalBox)
 
 			+ SVerticalBox::Slot()
-			.AutoHeight()
-			[
-				Details1View.ToSharedRef()
-			]
-
-			+ SVerticalBox::Slot()
-			.AutoHeight()
+			.FillHeight(1.0f)
+			.Padding(4, 4, 4, 4)
 			[
-				Details2View.ToSharedRef()
+				SNew(SScrollBox)
+				+SScrollBox::Slot()
+				[
+					Details1View.ToSharedRef()
+				]
+				+SScrollBox::Slot()
+				[
+					Details2View.ToSharedRef()
+				]
 			]
 
 			+ SVerticalBox::Slot()
@@ -211,7 +214,7 @@ void LevelSequenceEditorHelpers::OpenMasterSequenceDialog(const TSharedRef<FTabM
 			.HasCloseButton(true)
 			.SupportsMaximize(false)
 			.SupportsMinimize(false)
-			.ClientSize(FVector2D(600, 540));
+			.ClientSize(FVector2D(600, 600));
 
 		TSharedPtr<SDockTab> OwnerTab = TabManager->GetOwnerTab();
 		TSharedPtr<SWindow> RootWindow = OwnerTab.IsValid() ? OwnerTab->GetParentWindow() : TSharedPtr<SWindow>();
diff --git a/Engine/Plugins/MovieScene/MatineeToLevelSequence/Source/MatineeToLevelSequence/Private/MatineeToLevelSequenceModule.cpp b/Engine/Plugins/MovieScene/MatineeToLevelSequence/Source/MatineeToLevelSequence/Private/MatineeToLevelSequenceModule.cpp
--- a/Engine/Plugins/MovieScene/MatineeToLevelSequence/Source/MatineeToLevelSequence/Private/MatineeToLevelSequenceModule.cpp
+++ b/Engine/Plugins/MovieScene/MatineeToLevelSequence/Source/MatineeToLevelSequence/Private/MatineeToLevelSequenceModule.cpp
@@ -289,7 +289,7 @@ protected:
 				NewSequence->BindPossessableObject(ObjectGuid, *PropObject, BindingContext);
 			}
 
-			// cbb: String manipulations to get the property path in the rigth form for sequencer
+			// cbb: String manipulations to get the property path in the right form for sequencer
 			FString PropertyName = Property->GetFName().ToString();
 
 			// Special case for Light components which have some deprecated names
@@ -306,36 +306,31 @@ protected:
 				}
 			}
 
-			// Strip the object part of the property path
-			FString PropertyPath = Property->GetPathName();
-			int32 DotPos = INDEX_NONE;
-			if (PropertyPath.FindLastChar(TEXT('.'), DotPos))
+			TArray<UObject*> PropertyArray;
+			UObject* Outer = Property->GetOuter();
+			while (Outer->IsA(UProperty::StaticClass()) || Outer->IsA(UScriptStruct::StaticClass()))
 			{
-				PropertyPath = PropertyPath.RightChop(DotPos+1);
+				PropertyArray.Insert(Outer, 0);
+				Outer = Outer->GetOuter();
 			}
-
-			// Split it into the components
-			TArray<FString> PropertyPaths;
-			PropertyPath.ParseIntoArray(PropertyPaths, TEXT(":"));
-
-			// Reassemble path with "." separators
-			FString NewPropertyPath;
-			for (int32 PropertyIndex=0; PropertyIndex<PropertyPaths.Num()-1; ++PropertyIndex)
-			{		
-				if (NewPropertyPath.Len())
+			
+			FString PropertyPath;
+			for (auto PropertyIt : PropertyArray)
+			{
+				if (PropertyPath.Len())
 				{
-					NewPropertyPath = NewPropertyPath + TEXT(".");
+					PropertyPath = PropertyPath + TEXT(".");
 				}
-				NewPropertyPath = NewPropertyPath + PropertyPaths[PropertyIndex];
+				PropertyPath = PropertyPath + PropertyIt->GetName();
 			}
-			if (NewPropertyPath.Len())
+			if (PropertyPath.Len())
 			{
-				NewPropertyPath = NewPropertyPath + TEXT(".");
+				PropertyPath = PropertyPath + TEXT(".");
 			}
-			NewPropertyPath = NewPropertyPath + PropertyName;
+			PropertyPath = PropertyPath + PropertyName;
 
 			PropertyTrack = NewMovieScene->AddTrack<T>(ObjectGuid);	
-			PropertyTrack->SetPropertyNameAndPath(*PropertyName, NewPropertyPath);
+			PropertyTrack->SetPropertyNameAndPath(*PropertyName, PropertyPath);
 		}
 		else
 		{
diff --git a/Engine/Source/Editor/EditorStyle/Private/SlateEditorStyle.cpp b/Engine/Source/Editor/EditorStyle/Private/SlateEditorStyle.cpp
--- a/Engine/Source/Editor/EditorStyle/Private/SlateEditorStyle.cpp
+++ b/Engine/Source/Editor/EditorStyle/Private/SlateEditorStyle.cpp
@@ -5752,7 +5752,10 @@ void FSlateEditorStyle::FStyle::SetupClassIconsAndThumbnails()
 			TEXT("Cube"),
 			TEXT("Sphere"),
 			TEXT("Cylinder"),
-			TEXT("Cone")
+			TEXT("Cone"),
+			TEXT("CineCameraActor"),
+			TEXT("CameraRig_Crane"),
+			TEXT("CameraRig_Rail")
 		};
 
 		const TCHAR* Path = TEXT("Icons/AssetIcons/%s_%dx");
diff --git a/Engine/Source/Editor/MovieSceneCaptureDialog/Private/MovieSceneCaptureDialogModule.cpp b/Engine/Source/Editor/MovieSceneCaptureDialog/Private/MovieSceneCaptureDialogModule.cpp
--- a/Engine/Source/Editor/MovieSceneCaptureDialog/Private/MovieSceneCaptureDialogModule.cpp
+++ b/Engine/Source/Editor/MovieSceneCaptureDialog/Private/MovieSceneCaptureDialogModule.cpp
@@ -4,6 +4,7 @@
 #include "MovieSceneCapture.h"
 
 #include "UnrealEdMisc.h"
+#include "GameFramework/GameMode.h"
 
 #include "SlateBasics.h"
 #include "SlateExtras.h"
@@ -489,6 +490,12 @@ private:
 					
 					Window->Resize(PreviewWindowSize);
 
+					if (CaptureObject->Settings.GameModeOverride != nullptr)
+					{
+						CachedGameMode = CapturingFromWorld->GetWorldSettings()->DefaultGameMode;
+						CapturingFromWorld->GetWorldSettings()->DefaultGameMode = CaptureObject->Settings.GameModeOverride;
+					}
+
 					CaptureObject->Initialize(SlatePlayInEditorSession->SlatePlayInEditorWindowViewport, Context.PIEInstance);
 					OnStarted();
 				}
@@ -522,6 +529,11 @@ private:
 			}
 		}
 
+		if (CaptureObject->Settings.GameModeOverride != nullptr)
+		{
+			CapturingFromWorld->GetWorldSettings()->DefaultGameMode = CachedGameMode;
+		}
+
 		FObjectReader(GetMutableDefault<ULevelEditorPlaySettings>(), BackedUpPlaySettings);
 
 		FAudioDevice* AudioDevice = GEngine->GetMainAudioDevice();
@@ -558,6 +570,8 @@ private:
 	int32 BackedUpUseFixedPoolSize;
 	TArray<uint8> BackedUpPlaySettings;
 	UMovieSceneCapture* CaptureObject;
+
+	TSubclassOf<AGameMode> CachedGameMode;
 };
 
 class FMovieSceneCaptureDialogModule : public IMovieSceneCaptureDialogModule
diff --git a/Engine/Source/Editor/MovieSceneTools/Private/CurveKeyAreas/BoolKeyArea.cpp b/Engine/Source/Editor/MovieSceneTools/Private/CurveKeyAreas/BoolKeyArea.cpp
--- a/Engine/Source/Editor/MovieSceneTools/Private/CurveKeyAreas/BoolKeyArea.cpp
+++ b/Engine/Source/Editor/MovieSceneTools/Private/CurveKeyAreas/BoolKeyArea.cpp
@@ -20,14 +20,11 @@ TSharedRef<SWidget> FBoolKeyArea::CreateKeyEditor(ISequencer* Sequencer)
 		.Sequencer(Sequencer)
 		.OwningSection(OwningSection)
 		.Curve(&Curve)
-		.OnValueChanged(this, &FBoolKeyArea::OnValueChanged)
-		.IntermediateValue_Lambda([this] {
-			return IntermediateValue;
-		});
+		.ExternalValue(ExternalValue);
 };
 
-void FBoolKeyArea::OnValueChanged(bool InValue)
+bool FBoolKeyArea::ConvertCurveValueToIntegralType(int32 CurveValue) const
 {
-	ClearIntermediateValue();
+	return CurveValue != 0;
 }
 
diff --git a/Engine/Source/Editor/MovieSceneTools/Private/CurveKeyAreas/ByteKeyArea.cpp b/Engine/Source/Editor/MovieSceneTools/Private/CurveKeyAreas/ByteKeyArea.cpp
--- a/Engine/Source/Editor/MovieSceneTools/Private/CurveKeyAreas/ByteKeyArea.cpp
+++ b/Engine/Source/Editor/MovieSceneTools/Private/CurveKeyAreas/ByteKeyArea.cpp
@@ -16,11 +16,14 @@ bool FByteKeyArea::CanCreateKeyEditor()
 
 TSharedRef<SWidget> FByteKeyArea::CreateKeyEditor(ISequencer* Sequencer)
 {
-	return SNew(SIntegralCurveKeyEditor)
+	return SNew(SIntegralCurveKeyEditor<uint8>)
 		.Sequencer(Sequencer)
 		.OwningSection(OwningSection)
 		.Curve(&Curve)
-		.IntermediateValue_Lambda( [this] {
-			return IntermediateValue;
-		});
+		.ExternalValue(ExternalValue);
 };
+
+uint8 FByteKeyArea::ConvertCurveValueToIntegralType(int32 CurveValue) const
+{
+	return (uint8)FMath::Clamp<int32>(CurveValue, 0, TNumericLimits<uint8>::Max());
+}
diff --git a/Engine/Source/Editor/MovieSceneTools/Private/CurveKeyAreas/EnumKeyArea.cpp b/Engine/Source/Editor/MovieSceneTools/Private/CurveKeyAreas/EnumKeyArea.cpp
--- a/Engine/Source/Editor/MovieSceneTools/Private/CurveKeyAreas/EnumKeyArea.cpp
+++ b/Engine/Source/Editor/MovieSceneTools/Private/CurveKeyAreas/EnumKeyArea.cpp
@@ -20,15 +20,7 @@ TSharedRef<SWidget> FEnumKeyArea::CreateKeyEditor(ISequencer* Sequencer)
 		.Sequencer(Sequencer)
 		.OwningSection(OwningSection)
 		.Curve(&Curve)
-		.OnValueChanged(this, &FEnumKeyArea::OnValueChanged)
 		.Enum(Enum)
-		.IntermediateValue_Lambda([this] {
-			return IntermediateValue;
-		});
+		.ExternalValue(ExternalValue);
 };
 
-void FEnumKeyArea::OnValueChanged(int32 InValue)
-{
-	ClearIntermediateValue();
-}
-
diff --git a/Engine/Source/Editor/MovieSceneTools/Private/CurveKeyAreas/FloatCurveKeyArea.cpp b/Engine/Source/Editor/MovieSceneTools/Private/CurveKeyAreas/FloatCurveKeyArea.cpp
--- a/Engine/Source/Editor/MovieSceneTools/Private/CurveKeyAreas/FloatCurveKeyArea.cpp
+++ b/Engine/Source/Editor/MovieSceneTools/Private/CurveKeyAreas/FloatCurveKeyArea.cpp
@@ -26,15 +26,16 @@ TArray<FKeyHandle> FFloatCurveKeyArea::AddKeyUnique(float Time, EMovieSceneKeyIn
 			OwningSection->SetEndTime(Time);
 		}
 
-		float Value = Curve->Eval(Time);
-
-		if (TimeToCopyFrom != FLT_MAX)
+		float Value;
+		if (ExternalValue.IsSet() && ExternalValue.Get().IsSet() && TimeToCopyFrom == FLT_MAX)
 		{
-			Value = Curve->Eval(TimeToCopyFrom);
+			Value = ExternalValue.Get().GetValue();
 		}
-		else if ( IntermediateValue.IsSet() )
+		else
 		{
-			Value = IntermediateValue.GetValue();
+			float EvalTime = TimeToCopyFrom != FLT_MAX ? Time : TimeToCopyFrom;
+			float DefaultValue = 0;
+			Value = Curve->Eval(EvalTime, DefaultValue);
 		}
 
 		Curve->AddKey(Time, Value, false, CurrentKeyHandle);
@@ -58,10 +59,10 @@ TArray<FKeyHandle> FFloatCurveKeyArea::AddKeyUnique(float Time, EMovieSceneKeyIn
 			CurrentKey.LeaveTangentWeight = KeyToCopy.LeaveTangentWeight;
 		}
 	}
-	else if ( IntermediateValue.IsSet() )
+	else if (ExternalValue.IsSet() && ExternalValue.Get().IsSet())
 	{
-		float Value = IntermediateValue.GetValue();
-		Curve->UpdateOrAddKey(Time,Value);
+		float Value = ExternalValue.Get().GetValue();
+		Curve->UpdateOrAddKey(Time, Value);
 	}
 
 	return AddedKeyHandles;
@@ -95,10 +96,7 @@ TSharedRef<SWidget> FFloatCurveKeyArea::CreateKeyEditor(ISequencer* Sequencer)
 		.Sequencer(Sequencer)
 		.OwningSection(OwningSection)
 		.Curve(Curve)
-		.OnValueChanged(this, &FFloatCurveKeyArea::OnValueChanged)
-		.IntermediateValue_Lambda([this] {
-			return IntermediateValue;
-		});
+		.ExternalValue(ExternalValue);
 };
 
 
@@ -298,8 +296,3 @@ void FFloatCurveKeyArea::PasteKeys(const FMovieSceneClipboardKeyTrack& KeyTrack,
 		return true;
 	});
 }
-
-void FFloatCurveKeyArea::OnValueChanged(float InValue)
-{
-	ClearIntermediateValue();
-}
diff --git a/Engine/Source/Editor/MovieSceneTools/Private/CurveKeyAreas/IntegralKeyArea.cpp b/Engine/Source/Editor/MovieSceneTools/Private/CurveKeyAreas/IntegralKeyArea.cpp
--- a/Engine/Source/Editor/MovieSceneTools/Private/CurveKeyAreas/IntegralKeyArea.cpp
+++ b/Engine/Source/Editor/MovieSceneTools/Private/CurveKeyAreas/IntegralKeyArea.cpp
@@ -29,7 +29,7 @@ TArray<FKeyHandle> FIntegralCurveKeyAreaBase::AddKeyUnique(float Time, EMovieSce
 	}
 	else
 	{
-		SetKeyValue(Time);
+		UpdateKeyWithExternalValue(Time);
 	}
 
 	return AddedKeyHandles;
diff --git a/Engine/Source/Editor/MovieSceneTools/Private/CurveKeyAreas/StringCurveKeyArea.cpp b/Engine/Source/Editor/MovieSceneTools/Private/CurveKeyAreas/StringCurveKeyArea.cpp
--- a/Engine/Source/Editor/MovieSceneTools/Private/CurveKeyAreas/StringCurveKeyArea.cpp
+++ b/Engine/Source/Editor/MovieSceneTools/Private/CurveKeyAreas/StringCurveKeyArea.cpp
@@ -27,24 +27,25 @@ TArray<FKeyHandle> FStringCurveKeyArea::AddKeyUnique(float Time, EMovieSceneKeyI
 			OwningSection->SetEndTime(Time);
 		}
 
-		FString Value = Curve->Eval(Time, Curve->GetDefaultValue());
-
-		if (TimeToCopyFrom != FLT_MAX)
+		FString Value;
+		if (ExternalValue.IsSet() && ExternalValue.Get().IsSet() && TimeToCopyFrom == FLT_MAX)
 		{
-			Value = Curve->Eval(TimeToCopyFrom, Curve->GetDefaultValue());
+			Value = ExternalValue.Get().GetValue();
 		}
-		else if ( IntermediateValue.IsSet() )
+		else
 		{
-			Value = IntermediateValue.GetValue();
+			float EvalTime = TimeToCopyFrom != FLT_MAX ? Time : TimeToCopyFrom;
+			FString DefaultValue;
+			Value = Curve->Eval(EvalTime, DefaultValue);
 		}
 
 		Curve->AddKey(Time, Value, CurrentKeyHandle);
 		AddedKeyHandles.Add(CurrentKeyHandle);
 	}
-	else if ( IntermediateValue.IsSet() )
+	else if (ExternalValue.IsSet() && ExternalValue.Get().IsSet())
 	{
-		FString Value = IntermediateValue.GetValue();
-		Curve->UpdateOrAddKey(Time,Value);
+		FString Value = ExternalValue.Get().GetValue();
+		Curve->UpdateOrAddKey(Time, Value);
 	}
 
 	return AddedKeyHandles;
@@ -78,10 +79,7 @@ TSharedRef<SWidget> FStringCurveKeyArea::CreateKeyEditor(ISequencer* Sequencer)
 		.Sequencer(Sequencer)
 		.OwningSection(OwningSection)
 		.Curve(Curve)
-		.OnValueChanged(this, &FStringCurveKeyArea::OnValueChanged)
-		.IntermediateValue_Lambda([this] {
-			return IntermediateValue;
-		});
+		.ExternalValue(ExternalValue);
 };
 
 
@@ -244,8 +242,3 @@ void FStringCurveKeyArea::PasteKeys(const FMovieSceneClipboardKeyTrack& KeyTrack
 		return true;
 	});
 }
-
-void FStringCurveKeyArea::OnValueChanged(FString InValue)
-{
-	ClearIntermediateValue();
-}
diff --git a/Engine/Source/Editor/MovieSceneTools/Private/CurveKeyEditors/SBoolCurveKeyEditor.cpp b/Engine/Source/Editor/MovieSceneTools/Private/CurveKeyEditors/SBoolCurveKeyEditor.cpp
--- a/Engine/Source/Editor/MovieSceneTools/Private/CurveKeyEditors/SBoolCurveKeyEditor.cpp
+++ b/Engine/Source/Editor/MovieSceneTools/Private/CurveKeyEditors/SBoolCurveKeyEditor.cpp
@@ -1,7 +1,6 @@
 // Copyright 1998-2016 Epic Games, Inc. All Rights Reserved.
 
 #include "MovieSceneToolsPrivatePCH.h"
-#include "MovieSceneBoolSection.h"
 #include "SBoolCurveKeyEditor.h"
 
 #define LOCTEXT_NAMESPACE "BoolCurveKeyEditor"
@@ -11,9 +10,8 @@ void SBoolCurveKeyEditor::Construct(const FArguments& InArgs)
 	Sequencer = InArgs._Sequencer;
 	OwningSection = InArgs._OwningSection;
 	Curve = InArgs._Curve;
-	OnValueChangedEvent = InArgs._OnValueChanged;
-	IntermediateValue = InArgs._IntermediateValue;
-	float CurrentTime = Sequencer->GetCurrentLocalTime(*Sequencer->GetFocusedMovieSceneSequence());
+	ExternalValue = InArgs._ExternalValue;
+
 	ChildSlot
 	[
 		SNew(SCheckBox)
@@ -24,17 +22,19 @@ void SBoolCurveKeyEditor::Construct(const FArguments& InArgs)
 
 ECheckBoxState SBoolCurveKeyEditor::IsChecked() const
 {
-	if ( IntermediateValue.IsSet() && IntermediateValue.Get().IsSet() )
+	bool bCurrentValue;
+	if (ExternalValue.IsSet() && ExternalValue.Get().IsSet())
 	{
-		return IntermediateValue.Get().GetValue() ? ECheckBoxState::Checked : ECheckBoxState::Unchecked;
+		bCurrentValue = ExternalValue.Get().GetValue();
+	}
+	else
+	{
+		float CurrentTime = Sequencer->GetCurrentLocalTime(*Sequencer->GetFocusedMovieSceneSequence());
+		bool DefaultValue = false;
+		bCurrentValue = Curve->Evaluate(CurrentTime, DefaultValue) != 0;
 	}
-
-	float CurrentTime = Sequencer->GetCurrentLocalTime(*Sequencer->GetFocusedMovieSceneSequence());
-
-	UMovieSceneBoolSection* BoolSection = Cast<UMovieSceneBoolSection>(OwningSection);
-	bool bChecked = BoolSection ? BoolSection->Eval(CurrentTime) : Curve->Evaluate(CurrentTime) != 0;
 	
-	return bChecked ? ECheckBoxState::Checked : ECheckBoxState::Unchecked;
+	return bCurrentValue ? ECheckBoxState::Checked : ECheckBoxState::Unchecked;
 }
 
 void SBoolCurveKeyEditor::OnCheckStateChanged(ECheckBoxState NewCheckboxState)
@@ -44,10 +44,23 @@ void SBoolCurveKeyEditor::OnCheckStateChanged(ECheckBoxState NewCheckboxState)
 	if (OwningSection->TryModify())
 	{
 		float CurrentTime = Sequencer->GetCurrentLocalTime(*Sequencer->GetFocusedMovieSceneSequence());
+		bool bAutoSetTrackDefaults = Sequencer->GetAutoSetTrackDefaults();
+		int32 NewValue = NewCheckboxState == ECheckBoxState::Checked ? 1 : 0;
 
-		bool bKeyWillBeAdded = Curve->IsKeyHandleValid(Curve->FindKey(CurrentTime)) == false;
-		if (bKeyWillBeAdded)
+		FKeyHandle CurrentKeyHandle = Curve->FindKey(CurrentTime);
+		if (Curve->IsKeyHandleValid(CurrentKeyHandle))
+		{
+			Curve->SetKeyValue(CurrentKeyHandle, NewValue);
+		}
+		else
 		{
+			if (Curve->GetNumKeys() != 0 || bAutoSetTrackDefaults == false)
+			{
+				// When auto setting track defaults are disabled, add a key even when it's empty so that the changed
+				// value is saved and is propagated to the property.
+				Curve->AddKey(CurrentTime, NewValue, CurrentKeyHandle);
+			}
+
 			if (OwningSection->GetStartTime() > CurrentTime)
 			{
 				OwningSection->SetStartTime(CurrentTime);
@@ -57,20 +70,15 @@ void SBoolCurveKeyEditor::OnCheckStateChanged(ECheckBoxState NewCheckboxState)
 				OwningSection->SetEndTime(CurrentTime);
 			}
 		}
-		
-		int32 NewValue = NewCheckboxState == ECheckBoxState::Checked ? 1 : 0;
-		if (Curve->GetNumKeys() == 0)
+
+		// Always update the default value when auto-set default values is enabled so that the last changes
+		// are always saved to the track.
+		if (bAutoSetTrackDefaults)
 		{
 			Curve->SetDefaultValue(NewValue);
 		}
-		else
-		{
-			Curve->UpdateOrAddKey(CurrentTime, NewValue);
-		}
-
-		OnValueChangedEvent.ExecuteIfBound(NewCheckboxState == ECheckBoxState::Checked);
 
-		Sequencer->NotifyMovieSceneDataChanged( EMovieSceneDataChangeType::TrackValueChanged );
+		Sequencer->NotifyMovieSceneDataChanged( EMovieSceneDataChangeType::TrackValueChangedRefreshImmediately );
 	}
 }
 
diff --git a/Engine/Source/Editor/MovieSceneTools/Private/CurveKeyEditors/SBoolCurveKeyEditor.h b/Engine/Source/Editor/MovieSceneTools/Private/CurveKeyEditors/SBoolCurveKeyEditor.h
--- a/Engine/Source/Editor/MovieSceneTools/Private/CurveKeyEditors/SBoolCurveKeyEditor.h
+++ b/Engine/Source/Editor/MovieSceneTools/Private/CurveKeyEditors/SBoolCurveKeyEditor.h
@@ -8,15 +8,21 @@
 class SBoolCurveKeyEditor : public SCompoundWidget
 {
 public:
-	/** Notification for bool value change */
-	DECLARE_DELEGATE_OneParam( FOnValueChanged, bool );
-
 	SLATE_BEGIN_ARGS(SBoolCurveKeyEditor) {}
+
+		/** The sequencer which is controlling this key editor. */
 		SLATE_ARGUMENT(ISequencer*, Sequencer)
+
+		/** The section that owns the data edited by this key editor. */
 		SLATE_ARGUMENT(UMovieSceneSection*, OwningSection)
+
+		/** The curve being edited by this curve editor. */
 		SLATE_ARGUMENT(FIntegralCurve*, Curve)
-		SLATE_EVENT(FOnValueChanged, OnValueChanged)
-		SLATE_ATTRIBUTE(TOptional<bool>, IntermediateValue)
+
+		/** Allows the value displayed and edited by this key editor to be supplied from an external source.  This
+			is useful for curves on property tracks who's property value can change without changing the animation. */
+		SLATE_ATTRIBUTE(TOptional<bool>, ExternalValue)
+
 	SLATE_END_ARGS();
 
 	void Construct(const FArguments& InArgs);
@@ -29,6 +35,5 @@ private:
 	ISequencer* Sequencer;
 	UMovieSceneSection* OwningSection;
 	FIntegralCurve* Curve;
-	FOnValueChanged OnValueChangedEvent;
-	TAttribute<TOptional<bool>> IntermediateValue;
+	TAttribute<TOptional<bool>> ExternalValue;
 };
\ No newline at end of file
diff --git a/Engine/Source/Editor/MovieSceneTools/Private/CurveKeyEditors/SEnumCurveKeyEditor.cpp b/Engine/Source/Editor/MovieSceneTools/Private/CurveKeyEditors/SEnumCurveKeyEditor.cpp
--- a/Engine/Source/Editor/MovieSceneTools/Private/CurveKeyEditors/SEnumCurveKeyEditor.cpp
+++ b/Engine/Source/Editor/MovieSceneTools/Private/CurveKeyEditors/SEnumCurveKeyEditor.cpp
@@ -13,15 +13,14 @@ void SEnumCurveKeyEditor::Construct(const FArguments& InArgs)
 	Sequencer = InArgs._Sequencer;
 	OwningSection = InArgs._OwningSection;
 	Curve = InArgs._Curve;
-	OnValueChangedEvent = InArgs._OnValueChanged;
+	ExternalValue = InArgs._ExternalValue;
 
 	ChildSlot
 	[
 		MovieSceneToolHelpers::MakeEnumComboBox(
 			InArgs._Enum,
 			TAttribute<int32>::Create(TAttribute<int32>::FGetter::CreateSP(this, &SEnumCurveKeyEditor::OnGetCurrentValue)),
-			FOnEnumSelectionChanged::CreateSP(this, &SEnumCurveKeyEditor::OnComboSelectionChanged),
-			InArgs._IntermediateValue
+			FOnEnumSelectionChanged::CreateSP(this, &SEnumCurveKeyEditor::OnComboSelectionChanged)
 		)
 	];
 }
@@ -29,8 +28,14 @@ void SEnumCurveKeyEditor::Construct(const FArguments& InArgs)
 
 int32 SEnumCurveKeyEditor::OnGetCurrentValue() const
 {
+	if (ExternalValue.IsSet() && ExternalValue.Get().IsSet())
+	{
+		return ExternalValue.Get().GetValue();
+	}
+
 	float CurrentTime = Sequencer->GetCurrentLocalTime(*Sequencer->GetFocusedMovieSceneSequence());
-	return Curve->Evaluate(CurrentTime);
+	int32 DefaultValue = 0;
+	return Curve->Evaluate(CurrentTime, DefaultValue);
 }
 
 
@@ -42,34 +47,40 @@ void SEnumCurveKeyEditor::OnComboSelectionChanged(int32 InSelectedItem, ESelectI
 	if (OwningSection->TryModify())
 	{
 		float CurrentTime = Sequencer->GetCurrentLocalTime(*Sequencer->GetFocusedMovieSceneSequence());
+		bool bAutoSetTrackDefaults = Sequencer->GetAutoSetTrackDefaults();
 
-		bool bKeyWillBeAdded = Curve->IsKeyHandleValid(Curve->FindKey(CurrentTime)) == false;
-
-		if (bKeyWillBeAdded)
+		FKeyHandle CurrentKeyHandle = Curve->FindKey(CurrentTime);
+		if (Curve->IsKeyHandleValid(CurrentKeyHandle))
 		{
+			Curve->SetKeyValue(CurrentKeyHandle, InSelectedItem);
+		}
+		else
+		{
+			if (Curve->GetNumKeys() != 0 || bAutoSetTrackDefaults == false)
+			{
+				// When auto setting track defaults are disabled, add a key even when it's empty so that the changed
+				// value is saved and is propagated to the property.
+				Curve->AddKey(CurrentTime, InSelectedItem, CurrentKeyHandle);
+			}
+
 			if (OwningSection->GetStartTime() > CurrentTime)
 			{
 				OwningSection->SetStartTime(CurrentTime);
 			}
-
 			if (OwningSection->GetEndTime() < CurrentTime)
 			{
 				OwningSection->SetEndTime(CurrentTime);
 			}
 		}
 
-		if (Curve->GetNumKeys() == 0)
+		// Always update the default value when auto-set default values is enabled so that the last changes
+		// are always saved to the track.
+		if (bAutoSetTrackDefaults)
 		{
 			Curve->SetDefaultValue(InSelectedItem);
 		}
-		else
-		{
-			Curve->UpdateOrAddKey(CurrentTime, InSelectedItem);
-		}
-
-		OnValueChangedEvent.ExecuteIfBound(InSelectedItem);
 
-		Sequencer->NotifyMovieSceneDataChanged( EMovieSceneDataChangeType::TrackValueChanged );
+		Sequencer->NotifyMovieSceneDataChanged( EMovieSceneDataChangeType::TrackValueChangedRefreshImmediately );
 	}
 }
 
diff --git a/Engine/Source/Editor/MovieSceneTools/Private/CurveKeyEditors/SEnumCurveKeyEditor.h b/Engine/Source/Editor/MovieSceneTools/Private/CurveKeyEditors/SEnumCurveKeyEditor.h
--- a/Engine/Source/Editor/MovieSceneTools/Private/CurveKeyEditors/SEnumCurveKeyEditor.h
+++ b/Engine/Source/Editor/MovieSceneTools/Private/CurveKeyEditors/SEnumCurveKeyEditor.h
@@ -8,16 +8,24 @@
 class SEnumCurveKeyEditor : public SCompoundWidget
 {
 public:
-	/** Notification for enum value change */
-	DECLARE_DELEGATE_OneParam( FOnValueChanged, int32 );
-
 	SLATE_BEGIN_ARGS(SEnumCurveKeyEditor) {}
+
+		/** The sequencer which is controlling this key editor. */
 		SLATE_ARGUMENT(ISequencer*, Sequencer)
+
+		/** The section that owns the data edited by this key editor. */
 		SLATE_ARGUMENT(UMovieSceneSection*, OwningSection)
+
+		/** The curve being edited by this curve editor. */
 		SLATE_ARGUMENT(FIntegralCurve*, Curve)
-		SLATE_EVENT(FOnValueChanged, OnValueChanged)
+
+		/** The UEnum type which provides options for the curve being edited. */
 		SLATE_ARGUMENT(const UEnum*, Enum)
-		SLATE_ATTRIBUTE(TOptional<uint8>, IntermediateValue)
+
+		/** Allows the value displayed and edited by this key editor to be supplied from an external source.  This
+			is useful for curves on property tracks who's property value can change without changing the animation. */
+		SLATE_ATTRIBUTE(TOptional<uint8>, ExternalValue)
+
 	SLATE_END_ARGS();
 
 	void Construct(const FArguments& InArgs);
@@ -30,5 +38,5 @@ private:
 	ISequencer* Sequencer;
 	UMovieSceneSection* OwningSection;
 	FIntegralCurve* Curve;
-	FOnValueChanged OnValueChangedEvent;
+	TAttribute<TOptional<uint8>> ExternalValue;
 };
\ No newline at end of file
diff --git a/Engine/Source/Editor/MovieSceneTools/Private/CurveKeyEditors/SFloatCurveKeyEditor.cpp b/Engine/Source/Editor/MovieSceneTools/Private/CurveKeyEditors/SFloatCurveKeyEditor.cpp
--- a/Engine/Source/Editor/MovieSceneTools/Private/CurveKeyEditors/SFloatCurveKeyEditor.cpp
+++ b/Engine/Source/Editor/MovieSceneTools/Private/CurveKeyEditors/SFloatCurveKeyEditor.cpp
@@ -25,8 +25,8 @@ void SFloatCurveKeyEditor::Construct(const FArguments& InArgs)
 	Sequencer = InArgs._Sequencer;
 	OwningSection = InArgs._OwningSection;
 	Curve = InArgs._Curve;
-	OnValueChangedEvent = InArgs._OnValueChanged;
-	IntermediateValue = InArgs._IntermediateValue;
+	ExternalValue = InArgs._ExternalValue;
+
 	ChildSlot
 	[
 		SNew(SNonThrottledSpinBox<float>)
@@ -50,7 +50,6 @@ void SFloatCurveKeyEditor::OnBeginSliderMovement()
 {
 	GEditor->BeginTransaction(LOCTEXT("SetFloatKey", "Set Float Key Value"));
 	OwningSection->SetFlags(RF_Transactional);
-	OwningSection->TryModify();
 }
 
 void SFloatCurveKeyEditor::OnEndSliderMovement(float Value)
@@ -63,9 +62,9 @@ void SFloatCurveKeyEditor::OnEndSliderMovement(float Value)
 
 float SFloatCurveKeyEditor::OnGetKeyValue() const
 {
-	if ( IntermediateValue.IsSet() && IntermediateValue.Get().IsSet() )
+	if (ExternalValue.IsSet() && ExternalValue.Get().IsSet())
 	{
-		return IntermediateValue.Get().GetValue();
+		return ExternalValue.Get().GetValue();
 	}
 
 	float CurrentTime = Sequencer->GetCurrentLocalTime(*Sequencer->GetFocusedMovieSceneSequence());
@@ -77,6 +76,8 @@ void SFloatCurveKeyEditor::OnValueChanged(float Value)
 	if (OwningSection->TryModify())
 	{
 		float CurrentTime = Sequencer->GetCurrentLocalTime(*Sequencer->GetFocusedMovieSceneSequence());
+		bool bAutoSetTrackDefaults = Sequencer->GetAutoSetTrackDefaults();
+		
 		FKeyHandle CurrentKeyHandle = Curve->FindKey(CurrentTime);
 		if (Curve->IsKeyHandleValid(CurrentKeyHandle))
 		{
@@ -84,14 +85,11 @@ void SFloatCurveKeyEditor::OnValueChanged(float Value)
 		}
 		else
 		{
-			if (Curve->GetNumKeys() == 0)
-			{
-				Curve->SetDefaultValue(Value);
-			}
-			else
+			if (Curve->GetNumKeys() != 0 || bAutoSetTrackDefaults == false)
 			{
+				// When auto setting track defaults are disabled, add a key even when it's empty so that the changed
+				// value is saved and is propagated to the property.
 				Curve->AddKey(CurrentTime, Value, false, CurrentKeyHandle);
-
 				MovieSceneHelpers::SetKeyInterpolation(*Curve, CurrentKeyHandle, Sequencer->GetKeyInterpolation());
 			}
 
@@ -105,9 +103,14 @@ void SFloatCurveKeyEditor::OnValueChanged(float Value)
 			}
 		}
 
-		OnValueChangedEvent.ExecuteIfBound(Value);
+		// Always update the default value when auto-set default values is enabled so that the last changes
+		// are always saved to the track.
+		if (bAutoSetTrackDefaults)
+		{
+			Curve->SetDefaultValue(Value);
+		}
 
-		Sequencer->NotifyMovieSceneDataChanged( EMovieSceneDataChangeType::TrackValueChanged );
+		Sequencer->NotifyMovieSceneDataChanged(EMovieSceneDataChangeType::TrackValueChangedRefreshImmediately);
 	}
 }
 
@@ -116,7 +119,7 @@ void SFloatCurveKeyEditor::OnValueCommitted(float Value, ETextCommit::Type Commi
 	if (CommitInfo == ETextCommit::OnEnter || CommitInfo == ETextCommit::OnUserMovedFocus)
 	{
 		const FScopedTransaction Transaction( LOCTEXT("SetFloatKey", "Set Float Key Value") );
-
+		OwningSection->SetFlags(RF_Transactional);
 		OnValueChanged(Value);
 	}
 }
diff --git a/Engine/Source/Editor/MovieSceneTools/Private/CurveKeyEditors/SFloatCurveKeyEditor.h b/Engine/Source/Editor/MovieSceneTools/Private/CurveKeyEditors/SFloatCurveKeyEditor.h
--- a/Engine/Source/Editor/MovieSceneTools/Private/CurveKeyEditors/SFloatCurveKeyEditor.h
+++ b/Engine/Source/Editor/MovieSceneTools/Private/CurveKeyEditors/SFloatCurveKeyEditor.h
@@ -8,15 +8,21 @@
 class SFloatCurveKeyEditor : public SCompoundWidget
 {
 public:
-	/** Notification for numeric value change */
-	DECLARE_DELEGATE_OneParam( FOnValueChanged, float );
-
 	SLATE_BEGIN_ARGS(SFloatCurveKeyEditor) {}
+
+		/** The sequencer which is controlling this key editor. */
 		SLATE_ARGUMENT(ISequencer*, Sequencer)
+
+		/** The section that owns the data edited by this key editor. */
 		SLATE_ARGUMENT(UMovieSceneSection*, OwningSection)
+
+		/** The curve being edited by this curve editor. */
 		SLATE_ARGUMENT(FRichCurve*, Curve)
-		SLATE_EVENT(FOnValueChanged, OnValueChanged)
-		SLATE_ATTRIBUTE(TOptional<float>, IntermediateValue)
+
+		/** Allows the value displayed and edited by this key editor to be supplied from an external source.  This
+			is useful for curves on property tracks who's property value can change without changing the animation. */
+		SLATE_ATTRIBUTE(TOptional<float>, ExternalValue)
+
 	SLATE_END_ARGS();
 
 	void Construct(const FArguments& InArgs);
@@ -34,6 +40,5 @@ private:
 	ISequencer* Sequencer;
 	UMovieSceneSection* OwningSection;
 	FRichCurve* Curve;
-	FOnValueChanged OnValueChangedEvent;
-	TAttribute<TOptional<float>> IntermediateValue;
+	TAttribute<TOptional<float>> ExternalValue;
 };
\ No newline at end of file
diff --git a/Engine/Source/Editor/MovieSceneTools/Private/CurveKeyEditors/SIntegralCurveKeyEditor.h b/Engine/Source/Editor/MovieSceneTools/Private/CurveKeyEditors/SIntegralCurveKeyEditor.h
--- a/Engine/Source/Editor/MovieSceneTools/Private/CurveKeyEditors/SIntegralCurveKeyEditor.h
+++ b/Engine/Source/Editor/MovieSceneTools/Private/CurveKeyEditors/SIntegralCurveKeyEditor.h
@@ -2,33 +2,143 @@
 
 #pragma once
 
+#define LOCTEXT_NAMESPACE "IntegralCurveKeyEditor"
+
 /**
  * A widget for editing a curve representing integer keys.
  */
+template<typename IntegralType>
 class SIntegralCurveKeyEditor : public SCompoundWidget
 {
 public:
-	SLATE_BEGIN_ARGS(SIntegralCurveKeyEditor) {}
+	SLATE_BEGIN_ARGS(SIntegralCurveKeyEditor) { }
+
+		/** The sequencer which is controlling this key editor. */
 		SLATE_ARGUMENT(ISequencer*, Sequencer)
+
+		/** The section that owns the data edited by this key editor. */
 		SLATE_ARGUMENT(UMovieSceneSection*, OwningSection)
+
+		/** The curve being edited by this curve editor. */
 		SLATE_ARGUMENT(FIntegralCurve*, Curve)
-		SLATE_ATTRIBUTE(TOptional<uint8>, IntermediateValue)
+
+		/** Allows the value displayed and edited by this key editor to be supplied from an external source.  This
+		is useful for curves on property tracks who's property value can change without changing the animation. */
+		SLATE_ATTRIBUTE(TOptional<IntegralType>, ExternalValue)
+
 	SLATE_END_ARGS();
 
-	void Construct(const FArguments& InArgs);
+	void Construct(const FArguments& InArgs)
+	{
+		Sequencer = InArgs._Sequencer;
+		OwningSection = InArgs._OwningSection;
+		Curve = InArgs._Curve;
+		ExternalValue = InArgs._ExternalValue;
+
+		ChildSlot
+		[
+			SNew(SSpinBox<IntegralType>)
+			.Style(&FEditorStyle::GetWidgetStyle<FSpinBoxStyle>("Sequencer.HyperlinkSpinBox"))
+			.Font(FEditorStyle::GetFontStyle("Sequencer.AnimationOutliner.RegularFont"))
+			.MinValue(TOptional<IntegralType>())
+			.MaxValue(TOptional<IntegralType>())
+			.MaxSliderValue(TOptional<IntegralType>())
+			.MinSliderValue(TOptional<IntegralType>())
+			.Delta(1)
+			.Value(this, &SIntegralCurveKeyEditor::OnGetKeyValue)
+			.OnValueChanged(this, &SIntegralCurveKeyEditor::OnValueChanged)
+			.OnValueCommitted(this, &SIntegralCurveKeyEditor::OnValueCommitted)
+			.OnBeginSliderMovement(this, &SIntegralCurveKeyEditor::OnBeginSliderMovement)
+			.OnEndSliderMovement(this, &SIntegralCurveKeyEditor::OnEndSliderMovement)
+			.ClearKeyboardFocusOnCommit(true)
+		];
+	}
 
 private:
-	void OnBeginSliderMovement();
+	void OnBeginSliderMovement()
+	{
+		GEditor->BeginTransaction(LOCTEXT("SetIntegralKey", "Set Integral Key Value"));
+		OwningSection->SetFlags(RF_Transactional);
+		OwningSection->TryModify();
+	}
+
+	void OnEndSliderMovement(IntegralType Value)
+	{
+		if (GEditor->IsTransactionActive())
+		{
+			GEditor->EndTransaction();
+		}
+	}
+
+	IntegralType OnGetKeyValue() const
+	{
+		if (ExternalValue.IsSet() && ExternalValue.Get().IsSet())
+		{
+			return ExternalValue.Get().GetValue();
+		}
 
-	void OnEndSliderMovement(int32 Value);
+		float CurrentTime = Sequencer->GetCurrentLocalTime(*Sequencer->GetFocusedMovieSceneSequence());
+		return Curve->Evaluate(CurrentTime);
+	}
 
-	int32 OnGetKeyValue() const;
+	void OnValueChanged(IntegralType Value)
+	{
+		if (OwningSection->TryModify())
+		{
+			float CurrentTime = Sequencer->GetCurrentLocalTime(*Sequencer->GetFocusedMovieSceneSequence());
+			bool bAutoSetTrackDefaults = Sequencer->GetAutoSetTrackDefaults();
 
-	void OnValueChanged(int32 Value);
-	void OnValueCommitted(int32 Value, ETextCommit::Type CommitInfo);
+			FKeyHandle CurrentKeyHandle = Curve->FindKey(CurrentTime);
+			if (Curve->IsKeyHandleValid(CurrentKeyHandle))
+			{
+				Curve->SetKeyValue(CurrentKeyHandle, Value);
+			}
+			else
+			{
+				if (Curve->GetNumKeys() != 0 || bAutoSetTrackDefaults == false)
+				{
+					// When auto setting track defaults are disabled, add a key even when it's empty so that the changed
+					// value is saved and is propagated to the property.
+					Curve->AddKey(CurrentTime, Value,  CurrentKeyHandle);
+				}
+
+				if (OwningSection->GetStartTime() > CurrentTime)
+				{
+					OwningSection->SetStartTime(CurrentTime);
+				}
+				if (OwningSection->GetEndTime() < CurrentTime)
+				{
+					OwningSection->SetEndTime(CurrentTime);
+				}
+			}
+
+			// Always update the default value when auto-set default values is enabled so that the last changes
+			// are always saved to the track.
+			if (bAutoSetTrackDefaults)
+			{
+				Curve->SetDefaultValue(Value);
+			}
+
+			Sequencer->NotifyMovieSceneDataChanged(EMovieSceneDataChangeType::TrackValueChangedRefreshImmediately);
+		}
+	}
+
+	void OnValueCommitted(IntegralType Value, ETextCommit::Type CommitInfo)
+	{
+		if (CommitInfo == ETextCommit::OnEnter || CommitInfo == ETextCommit::OnUserMovedFocus)
+		{
+			const FScopedTransaction Transaction(LOCTEXT("SetIntegralKey", "Set Integral Key Value"));
+
+			OnValueChanged(Value);
+		}
+	}
+
+private:
 
 	ISequencer* Sequencer;
 	UMovieSceneSection* OwningSection;
 	FIntegralCurve* Curve;
-	TAttribute<TOptional<uint8>> IntermediateValue;
-};
\ No newline at end of file
+	TAttribute<TOptional<IntegralType>> ExternalValue;
+};
+
+#undef LOCTEXT_NAMESPACE
\ No newline at end of file
diff --git a/Engine/Source/Editor/MovieSceneTools/Private/CurveKeyEditors/SStringCurveKeyEditor.cpp b/Engine/Source/Editor/MovieSceneTools/Private/CurveKeyEditors/SStringCurveKeyEditor.cpp
--- a/Engine/Source/Editor/MovieSceneTools/Private/CurveKeyEditors/SStringCurveKeyEditor.cpp
+++ b/Engine/Source/Editor/MovieSceneTools/Private/CurveKeyEditors/SStringCurveKeyEditor.cpp
@@ -1,8 +1,8 @@
 // Copyright 1998t-2016 Epic Games, Inc. All Rights Reserved.
 
 #include "MovieSceneToolsPrivatePCH.h"
-#include "MovieSceneStringSection.h"
 #include "SStringCurveKeyEditor.h"
+#include "Curves/StringCurve.h"
 
 #define LOCTEXT_NAMESPACE "StringCurveKeyEditor"
 
@@ -11,8 +11,8 @@ void SStringCurveKeyEditor::Construct(const FArguments& InArgs)
 	Sequencer = InArgs._Sequencer;
 	OwningSection = InArgs._OwningSection;
 	Curve = InArgs._Curve;
-	OnValueChangedEvent = InArgs._OnValueChanged;
-	IntermediateValue = InArgs._IntermediateValue;
+	ExternalValue = InArgs._ExternalValue;
+
 	float CurrentTime = Sequencer->GetCurrentLocalTime(*Sequencer->GetFocusedMovieSceneSequence());
 	ChildSlot
 	[
@@ -25,15 +25,13 @@ void SStringCurveKeyEditor::Construct(const FArguments& InArgs)
 
 FText SStringCurveKeyEditor::GetText() const
 {
-	if ( IntermediateValue.IsSet() && IntermediateValue.Get().IsSet() )
+	if (ExternalValue.IsSet() && ExternalValue.Get().IsSet())
 	{
-		return FText::FromString(IntermediateValue.Get().GetValue());
+		return FText::FromString(ExternalValue.Get().GetValue());
 	}
-
 	float CurrentTime = Sequencer->GetCurrentLocalTime(*Sequencer->GetFocusedMovieSceneSequence());
-
-	UMovieSceneStringSection* StringSection = Cast<UMovieSceneStringSection>(OwningSection);
-	FString CurrentValue = StringSection ? StringSection->Eval(CurrentTime, Curve->GetDefaultValue()) : Curve->Eval(CurrentTime, Curve->GetDefaultValue());
+	FString DefaultValue;
+	FString CurrentValue = Curve->Eval(CurrentTime, DefaultValue);
 	
 	return FText::FromString(CurrentValue);
 }
@@ -45,10 +43,22 @@ void SStringCurveKeyEditor::OnTextCommitted(const FText& InText, ETextCommit::Ty
 	if (OwningSection->TryModify())
 	{
 		float CurrentTime = Sequencer->GetCurrentLocalTime(*Sequencer->GetFocusedMovieSceneSequence());
+		bool bAutoSetTrackDefaults = Sequencer->GetAutoSetTrackDefaults();
 
-		bool bKeyWillBeAdded = Curve->IsKeyHandleValid(Curve->FindKey(CurrentTime)) == false;
-		if (bKeyWillBeAdded)
+		FKeyHandle CurrentKeyHandle = Curve->FindKey(CurrentTime);
+		if (Curve->IsKeyHandleValid(CurrentKeyHandle))
 		{
+			Curve->SetKeyValue(CurrentKeyHandle, InText.ToString());
+		}
+		else
+		{
+			if (Curve->GetNumKeys() != 0 || bAutoSetTrackDefaults == false)
+			{
+				// When auto setting track defaults are disabled, add a key even when it's empty so that the changed
+				// value is saved and is propagated to the property.
+				Curve->AddKey(CurrentTime, InText.ToString(), CurrentKeyHandle);
+			}
+
 			if (OwningSection->GetStartTime() > CurrentTime)
 			{
 				OwningSection->SetStartTime(CurrentTime);
@@ -58,19 +68,15 @@ void SStringCurveKeyEditor::OnTextCommitted(const FText& InText, ETextCommit::Ty
 				OwningSection->SetEndTime(CurrentTime);
 			}
 		}
-		
-		if (Curve->GetNumKeys() == 0)
+
+		// Always update the default value when auto-set default values is enabled so that the last changes
+		// are always saved to the track.
+		if (bAutoSetTrackDefaults)
 		{
 			Curve->SetDefaultValue(InText.ToString());
 		}
-		else
-		{
-			Curve->UpdateOrAddKey(CurrentTime, InText.ToString());
-		}
-
-		OnValueChangedEvent.ExecuteIfBound(InText.ToString());
 
-		Sequencer->NotifyMovieSceneDataChanged( EMovieSceneDataChangeType::TrackValueChanged );
+		Sequencer->NotifyMovieSceneDataChanged( EMovieSceneDataChangeType::TrackValueChangedRefreshImmediately );
 	}
 }
 
diff --git a/Engine/Source/Editor/MovieSceneTools/Private/CurveKeyEditors/SStringCurveKeyEditor.h b/Engine/Source/Editor/MovieSceneTools/Private/CurveKeyEditors/SStringCurveKeyEditor.h
--- a/Engine/Source/Editor/MovieSceneTools/Private/CurveKeyEditors/SStringCurveKeyEditor.h
+++ b/Engine/Source/Editor/MovieSceneTools/Private/CurveKeyEditors/SStringCurveKeyEditor.h
@@ -2,21 +2,29 @@
 
 #pragma once
 
+struct FStringCurve;
+
 /**
  * A widget for editing a curve representing string keys.
  */
 class SStringCurveKeyEditor : public SCompoundWidget
 {
 public:
-	/** Notification for string value change */
-	DECLARE_DELEGATE_OneParam( FOnValueChanged, FString );
-
 	SLATE_BEGIN_ARGS(SStringCurveKeyEditor) {}
+
+		/** The sequencer which is controlling this key editor. */
 		SLATE_ARGUMENT(ISequencer*, Sequencer)
+
+		/** The section that owns the data edited by this key editor. */
 		SLATE_ARGUMENT(UMovieSceneSection*, OwningSection)
+
+		/** The curve being edited by this curve editor. */
 		SLATE_ARGUMENT(FStringCurve*, Curve)
-		SLATE_EVENT(FOnValueChanged, OnValueChanged)
-		SLATE_ATTRIBUTE(TOptional<FString>, IntermediateValue)
+
+		/** Allows the value displayed and edited by this key editor to be supplied from an external source.  This
+		is useful for curves on property tracks who's property value can change without changing the animation. */
+		SLATE_ATTRIBUTE(TOptional<FString>, ExternalValue)
+
 	SLATE_END_ARGS();
 
 	void Construct(const FArguments& InArgs);
@@ -29,6 +37,5 @@ private:
 	ISequencer* Sequencer;
 	UMovieSceneSection* OwningSection;
 	FStringCurve* Curve;
-	FOnValueChanged OnValueChangedEvent;
-	TAttribute<TOptional<FString>> IntermediateValue;
+	TAttribute<TOptional<FString>> ExternalValue;
 };
\ No newline at end of file
diff --git a/Engine/Source/Editor/MovieSceneTools/Private/MatineeImportTools.cpp b/Engine/Source/Editor/MovieSceneTools/Private/MatineeImportTools.cpp
--- a/Engine/Source/Editor/MovieSceneTools/Private/MatineeImportTools.cpp
+++ b/Engine/Source/Editor/MovieSceneTools/Private/MatineeImportTools.cpp
@@ -375,7 +375,7 @@ bool FMatineeImportTools::CopyInterpMoveTrack( UInterpTrackMove* MoveTrack, UMov
 					{
 						SubTrackCurve = &TranslationZCurve;
 					}
-					else if (MoveSubTrack->MoveAxis == EInterpMoveAxis::AXIS_RotationY)
+					else if (MoveSubTrack->MoveAxis == EInterpMoveAxis::AXIS_RotationX)
 					{
 						SubTrackCurve = &RotationXCurve;
 					}
diff --git a/Engine/Source/Editor/MovieSceneTools/Private/MovieSceneToolHelpers.cpp b/Engine/Source/Editor/MovieSceneTools/Private/MovieSceneToolHelpers.cpp
--- a/Engine/Source/Editor/MovieSceneTools/Private/MovieSceneToolHelpers.cpp
+++ b/Engine/Source/Editor/MovieSceneTools/Private/MovieSceneToolHelpers.cpp
@@ -403,7 +403,6 @@ class SEnumCombobox : public SComboBox<TSharedPtr<int32>>
 public:
 	SLATE_BEGIN_ARGS(SEnumCombobox) {}
 
-	SLATE_ATTRIBUTE(TOptional<uint8>, IntermediateValue)
 	SLATE_ATTRIBUTE(int32, CurrentValue)
 	SLATE_ARGUMENT(FOnEnumSelectionChanged, OnEnumSelectionChanged)
 
@@ -412,7 +411,6 @@ public:
 	void Construct(const FArguments& InArgs, const UEnum* InEnum)
 	{
 		Enum = InEnum;
-		IntermediateValue = InArgs._IntermediateValue;
 		CurrentValue = InArgs._CurrentValue;
 		check(CurrentValue.IsBound());
 		OnEnumSelectionChangedDelegate = InArgs._OnEnumSelectionChanged;
@@ -448,12 +446,6 @@ public:
 private:
 	FText GetCurrentValue() const
 	{
-		if(IntermediateValue.IsSet() && IntermediateValue.Get().IsSet())
-		{
-			int32 IntermediateNameIndex = Enum->GetIndexByValue(IntermediateValue.Get().GetValue());
-			return Enum->GetDisplayNameText(IntermediateNameIndex);
-		}
-
 		int32 CurrentNameIndex = Enum->GetIndexByValue(CurrentValue.Get());
 		return Enum->GetDisplayNameText(CurrentNameIndex);
 	}
@@ -497,8 +489,6 @@ private:
 
 	TAttribute<int32> CurrentValue;
 
-	TAttribute<TOptional<uint8>> IntermediateValue;
-
 	TArray<TSharedPtr<int32>> VisibleEnumNameIndices;
 
 	bool bUpdatingSelectionInternally;
@@ -506,12 +496,11 @@ private:
 	FOnEnumSelectionChanged OnEnumSelectionChangedDelegate;
 };
 
-TSharedRef<SWidget> MovieSceneToolHelpers::MakeEnumComboBox(const UEnum* InEnum, TAttribute<int32> InCurrentValue, FOnEnumSelectionChanged InOnSelectionChanged, TAttribute<TOptional<uint8>> InIntermediateValue)
+TSharedRef<SWidget> MovieSceneToolHelpers::MakeEnumComboBox(const UEnum* InEnum, TAttribute<int32> InCurrentValue, FOnEnumSelectionChanged InOnSelectionChanged)
 {
 	return SNew(SEnumCombobox, InEnum)
 		.CurrentValue(InCurrentValue)
-		.OnEnumSelectionChanged(InOnSelectionChanged)
-		.IntermediateValue(InIntermediateValue);
+		.OnEnumSelectionChanged(InOnSelectionChanged);
 }
 
 bool MovieSceneToolHelpers::ShowImportEDLDialog(UMovieScene* InMovieScene, float InFrameRate, FString InOpenDirectory)
diff --git a/Engine/Source/Editor/MovieSceneTools/Private/PropertySection.cpp b/Engine/Source/Editor/MovieSceneTools/Private/PropertySection.cpp
new file mode 100644
--- /dev/null
+++ b/Engine/Source/Editor/MovieSceneTools/Private/PropertySection.cpp
@@ -0,0 +1,80 @@
+// Copyright 1998-2016 Epic Games, Inc. All Rights Reserved.
+
+#pragma once
+
+#include "MovieSceneToolsPrivatePCH.h"
+#include "PropertySection.h"
+
+
+FPropertySection::FPropertySection(UMovieSceneSection& InSectionObject, const FText& InDisplayName)
+	: DisplayName(InDisplayName)
+	, SectionObject(InSectionObject)
+	, Sequencer(nullptr)
+{ }
+
+FPropertySection::FPropertySection(ISequencer* InSequencer, FGuid InObjectBinding, FName InPropertyName, 
+	const FString& InPropertyPath, UMovieSceneSection& InSectionObject, const FText& InDisplayName)
+	: DisplayName(InDisplayName)
+	, SectionObject(InSectionObject)
+	, Sequencer(InSequencer)
+	, ObjectBinding(InObjectBinding)
+	, PropertyBindings(MakeShareable(new FTrackInstancePropertyBindings(InPropertyName, InPropertyPath)))
+{
+
+}
+
+UMovieSceneSection* FPropertySection::GetSectionObject()
+{
+	return &SectionObject;
+}
+
+FText FPropertySection::GetDisplayName() const
+{
+	return DisplayName;
+}
+	
+FText FPropertySection::GetSectionTitle() const
+{
+	return FText::GetEmpty();
+}
+
+int32 FPropertySection::OnPaintSection( FSequencerSectionPainter& Painter ) const
+{
+	return Painter.PaintSectionBackground();
+}
+
+UObject* FPropertySection::GetRuntimeObjectAndUpdatePropertyBindings() const
+{
+	UObject* RuntimeObject = nullptr;
+	if (Sequencer != nullptr && PropertyBindings.IsValid())
+	{
+		TArray<TWeakObjectPtr<UObject>> RuntimeObjects;
+		Sequencer->GetRuntimeObjects(Sequencer->GetFocusedMovieSceneSequenceInstance(), ObjectBinding, RuntimeObjects);
+		if (RuntimeObjects.Num() == 1)
+		{
+			TWeakObjectPtr<UObject> RuntimeObjectPtr = RuntimeObjects[0];
+			if (RuntimeObjectPtr.IsValid())
+			{
+				if (RuntimeObjectPtr != RuntimeObjectCache)
+				{
+					PropertyBindings->UpdateBindings(RuntimeObjects);
+					RuntimeObjectCache = RuntimeObjectPtr;
+				}
+				RuntimeObject = RuntimeObjectPtr.Get();
+			}
+		}
+	}
+	return RuntimeObject;
+}
+
+UProperty* FPropertySection::GetProperty() const
+{
+	UObject* RuntimeObject = GetRuntimeObjectAndUpdatePropertyBindings();
+	return PropertyBindings.IsValid() && RuntimeObject != nullptr ? PropertyBindings->GetProperty(RuntimeObject) : nullptr;
+}
+
+bool FPropertySection::CanGetPropertyValue() const
+{
+	return Sequencer != nullptr && PropertyBindings.IsValid();
+}
+
diff --git a/Engine/Source/Editor/MovieSceneTools/Private/Sections/ActorReferencePropertySection.cpp b/Engine/Source/Editor/MovieSceneTools/Private/Sections/ActorReferencePropertySection.cpp
--- a/Engine/Source/Editor/MovieSceneTools/Private/Sections/ActorReferencePropertySection.cpp
+++ b/Engine/Source/Editor/MovieSceneTools/Private/Sections/ActorReferencePropertySection.cpp
@@ -4,18 +4,131 @@
 #include "ActorReferencePropertySection.h"
 #include "MovieSceneActorReferenceSection.h"
 
-void FActorReferencePropertySection::GenerateSectionLayout( class ISectionLayoutBuilder& LayoutBuilder ) const
+
+void FActorReferenceKeyArea::CopyKeys(FMovieSceneClipboardBuilder& ClipboardBuilder, const TFunctionRef<bool(FKeyHandle, const IKeyArea&)>& KeyMask) const
 {
-	UMovieSceneActorReferenceSection* ActorReferenceSection = Cast<UMovieSceneActorReferenceSection>( &SectionObject );
-	KeyArea = MakeShareable( new FActorReferenceKeyArea( ActorReferenceSection->GetActorReferenceCurve(), ActorReferenceSection ) );
-	LayoutBuilder.SetSectionAsKeyArea( KeyArea.ToSharedRef() );
+	const UMovieSceneActorReferenceSection* Section = Cast<UMovieSceneActorReferenceSection>(OwningSection);
+	UMovieSceneTrack* Track = Section ? Section->GetTypedOuter<UMovieSceneTrack>() : nullptr;
+	if (!Track)
+	{
+		return;
+	}
+
+	FMovieSceneClipboardKeyTrack* KeyTrack = nullptr;
+
+	for (auto It(Curve.GetKeyHandleIterator()); It; ++It)
+	{
+		FKeyHandle Handle = It.Key();
+		if (KeyMask(Handle, *this))
+		{
+			if (!KeyTrack)
+			{
+				KeyTrack = &ClipboardBuilder.FindOrAddKeyTrack<FGuid>(GetName(), *Track);
+			}
+
+			// We can't copy the curve values directly since the actor reference curve hold indices to
+			// the actual guids since there is no FGuid curve.
+			// TODO: Fix up this code once there is an FGuid curve.
+			float KeyTime = Curve.GetKeyTime(Handle);
+			FGuid KeyValue = Section->Eval(KeyTime);
+			KeyTrack->AddKey(KeyTime, KeyValue);
+		}
+	}
+}
+
+
+void FActorReferenceKeyArea::PasteKeys(const FMovieSceneClipboardKeyTrack& KeyTrack, const FMovieSceneClipboardEnvironment& SrcEnvironment, const FSequencerPasteEnvironment& DstEnvironment)
+{
+	float PasteAt = DstEnvironment.CardinalTime;
+
+	KeyTrack.IterateKeys([&](const FMovieSceneClipboardKey& Key) {
+		UMovieSceneActorReferenceSection* Section = Cast<UMovieSceneActorReferenceSection>(GetOwningSection());
+		if (!Section)
+		{
+			return true;
+		}
+
+		if (Section->TryModify())
+		{
+			float Time = PasteAt + Key.GetTime();
+			if (Section->GetStartTime() > Time)
+			{
+				Section->SetStartTime(Time);
+			}
+			if (Section->GetEndTime() < Time)
+			{
+				Section->SetEndTime(Time);
+			}
+
+			// We can't paste the curve values directly since the actor reference curve hold indices to
+			// the actual guids since there is no FGuid curve.
+			// TODO: Fix up this code once there is an FGuid curve.
+			FKeyHandle KeyHandle = Section->AddKey(Time, Key.GetValue<FGuid>());
+			DstEnvironment.ReportPastedKey(KeyHandle, *this);
+		}
+
+		return true;
+	});
+}
+
+
+void FActorReferenceKeyArea::EvaluateAndAddKey(float Time, float TimeToCopyFrom, FKeyHandle& CurrentKey)
+{
+	UMovieSceneActorReferenceSection* ActorReferenceSection = Cast<UMovieSceneActorReferenceSection>(OwningSection);
+	checkf(ActorReferenceSection, TEXT("Incompatible section in FActorReferenceKeyArea"));
+
+	// Add the keys to the section here instead of the curve because the actor reference section doesn't
+	// store the guids directly in the curve.
+	// TODO: Fix up this code once there is an FGuid curve.
+	if (ActorReferenceSection != nullptr)
+	{
+		FGuid Value;
+		if (ExternalValue.IsSet() && ExternalValue.Get().IsSet() && TimeToCopyFrom == FLT_MAX)
+		{
+			Value = ExternalValue.Get().GetValue();
+		}
+		else
+		{
+			float EvalTime = TimeToCopyFrom != FLT_MAX ? Time : TimeToCopyFrom;
+			Value = ActorReferenceSection->Eval(EvalTime);
+		}
+
+		ActorReferenceSection->AddKey(Time, Value);
+	}
 }
 
-void FActorReferencePropertySection::SetIntermediateValue( FPropertyChangedParams PropertyChangedParams )
+
+void FActorReferenceKeyArea::UpdateKeyWithExternalValue(float Time)
 {
+	UMovieSceneActorReferenceSection* ActorReferenceSection = Cast<UMovieSceneActorReferenceSection>(OwningSection);
+	checkf(ActorReferenceSection, TEXT("Incompatible section in FActorReferenceKeyArea"));
+
+	// Set the key through the section here instead of the curve because the actor reference section doesn't
+	// store the guids directly in the curve.
+	if (ExternalValue.IsSet() && ExternalValue.Get().IsSet())
+	{
+		ActorReferenceSection->AddKey(Time, ExternalValue.Get().GetValue());
+	}
 }
 
-void FActorReferencePropertySection::ClearIntermediateValue()
+
+void FActorReferencePropertySection::GenerateSectionLayout( class ISectionLayoutBuilder& LayoutBuilder ) const
+{
+	UMovieSceneActorReferenceSection* ActorReferenceSection = Cast<UMovieSceneActorReferenceSection>( &SectionObject );
+
+	TAttribute<TOptional<FGuid>> ActorGuidExternalValue = TAttribute<TOptional<FGuid>>::Create(
+		TAttribute<TOptional<FGuid>>::FGetter::CreateRaw(this, &FActorReferencePropertySection::GetActorGuid));
+	TSharedRef<FActorReferenceKeyArea> KeyArea = MakeShareable( new FActorReferenceKeyArea( ActorReferenceSection->GetActorReferenceCurve(), ActorGuidExternalValue, ActorReferenceSection ) );
+	LayoutBuilder.SetSectionAsKeyArea( KeyArea );
+}
+
+
+TOptional<FGuid> FActorReferencePropertySection::GetActorGuid() const
 {
-	KeyArea->ClearIntermediateValue();
+	TOptional<AActor*> CurrentActor = GetPropertyValue<AActor*>();
+	if (CurrentActor.IsSet() && CurrentActor.GetValue() != nullptr)
+	{
+		return TOptional<FGuid>(GetSequencer()->GetFocusedMovieSceneSequenceInstance()->FindObjectId(*CurrentActor.GetValue()));
+	}
+	return TOptional<FGuid>(FGuid());
 }
\ No newline at end of file
diff --git a/Engine/Source/Editor/MovieSceneTools/Private/Sections/ActorReferencePropertySection.h b/Engine/Source/Editor/MovieSceneTools/Private/Sections/ActorReferencePropertySection.h
--- a/Engine/Source/Editor/MovieSceneTools/Private/Sections/ActorReferencePropertySection.h
+++ b/Engine/Source/Editor/MovieSceneTools/Private/Sections/ActorReferencePropertySection.h
@@ -2,31 +2,54 @@
 
 #pragma once
 
-#include "IntegralKeyArea.h"
 #include "PropertySection.h"
+#include "IntegralKeyArea.h"
 
 
 struct FIntegralCurve;
-class FActorReferenceKeyArea;
 class UMovieSceneSection;
 
 
 /** A key area for displaying and editing an actor reference property. */
 class FActorReferenceKeyArea
-	: public FIntegralKeyArea<int32>
+	: public FIntegralCurveKeyAreaBase
 {
 public:
 
-	FActorReferenceKeyArea(FIntegralCurve& InCurve, UMovieSceneSection* InOwningSection)
-		: FIntegralKeyArea(InCurve, InOwningSection)
+	/**
+	 * Creates a new actor reference key area.
+	 *
+	 * @param InCurve The integral curve which has the actor reference keys.
+	 * @param InExternalValue An attribute which can provide an external value for this key area.  External values are
+	 *        useful for things like property tracks where the property value can change without changing the animation
+	 *        and we want to be able to key and update using the new property value.
+	 * @param InOwningSection The section which owns the curve which is being displayed and edited by this area.
+	 */
+	FActorReferenceKeyArea(FIntegralCurve& InCurve, TAttribute<TOptional<FGuid>> InExternalValue, UMovieSceneSection* InOwningSection)
+		: FIntegralCurveKeyAreaBase(InCurve, InOwningSection)
+		, ExternalValue(InExternalValue)
 	{ }
 
 public:
 
-	//~ FIntegralKeyArea overrides
+	//~ IKeyArea overrides
 
 	virtual bool CanCreateKeyEditor() override { return false; }
 	virtual TSharedRef<SWidget> CreateKeyEditor(ISequencer* Sequencer) override { return SNew(SSpacer); }
+	virtual void CopyKeys(FMovieSceneClipboardBuilder& ClipboardBuilder, const TFunctionRef<bool(FKeyHandle, const IKeyArea&)>& KeyMask) const override;
+	virtual void PasteKeys(const FMovieSceneClipboardKeyTrack& KeyTrack, const FMovieSceneClipboardEnvironment& SrcEnvironment, const FSequencerPasteEnvironment& DstEnvironment) override;
+
+
+protected:
+
+	//~ FIntegralCurveKeyAreaBase interface
+
+	virtual void EvaluateAndAddKey(float Time, float TimeToCopyFrom, FKeyHandle& CurrentKey) override;
+	virtual void UpdateKeyWithExternalValue(float Time) override;
+
+private:
+
+	TAttribute<TOptional<FGuid>> ExternalValue;
 };
 
 
@@ -38,19 +61,28 @@ class FActorReferencePropertySection
 {
 public:
 
-	FActorReferencePropertySection(UMovieSceneSection& InSectionObject, const FText& SectionName)
-		: FPropertySection(InSectionObject, SectionName)
-	{ }
+	/**
+	* Creates a new actor reference property section.
+	*
+	* @param InSequencer The sequencer which is controlling this property section.
+	* @param InObjectBinding The object binding for the object which owns the property that this section is animating.
+	* @param InPropertyName The name of the property which is animated by this section.
+	* @param InPropertyPath A string representing the path to the property which is animated by this section.
+	* @param InSectionObject The section object which is being displayed and edited.
+	* @param InDisplayName A display name for the section being displayed and edited.
+	*/
+	FActorReferencePropertySection(ISequencer* InSequencer, FGuid InObjectBinding, FName InPropertyName, const FString& InPropertyPath, UMovieSceneSection& InSectionObject, const FText& InDisplayName)
+		: FPropertySection(InSequencer, InObjectBinding, InPropertyName, InPropertyPath, InSectionObject, InDisplayName)
+	{
+	}
 
 public:
 
 	//~ FPropertySection interface
 
 	virtual void GenerateSectionLayout(class ISectionLayoutBuilder& LayoutBuilder) const override;
-	virtual void SetIntermediateValue(FPropertyChangedParams PropertyChangedParams) override;
-	virtual void ClearIntermediateValue() override;
 
 private:
 
-	mutable TSharedPtr<FActorReferenceKeyArea> KeyArea;
+	TOptional<FGuid> GetActorGuid() const;
 };
diff --git a/Engine/Source/Editor/MovieSceneTools/Private/Sections/BoolPropertySection.cpp b/Engine/Source/Editor/MovieSceneTools/Private/Sections/BoolPropertySection.cpp
--- a/Engine/Source/Editor/MovieSceneTools/Private/Sections/BoolPropertySection.cpp
+++ b/Engine/Source/Editor/MovieSceneTools/Private/Sections/BoolPropertySection.cpp
@@ -11,8 +11,16 @@
 void FBoolPropertySection::GenerateSectionLayout( class ISectionLayoutBuilder& LayoutBuilder ) const
 {
 	UMovieSceneBoolSection* BoolSection = Cast<UMovieSceneBoolSection>( &SectionObject );
-	KeyArea = MakeShareable( new FBoolKeyArea( BoolSection->GetCurve(), BoolSection ) );
-	LayoutBuilder.SetSectionAsKeyArea( KeyArea.ToSharedRef() );
+	TAttribute<TOptional<bool>> ExternalValue;
+	if (CanGetPropertyValue())
+	{
+		ExternalValue.Bind(TAttribute<TOptional<bool>>::FGetter::CreateLambda([&]
+		{
+			return GetPropertyValue<bool>();
+		}));
+	}
+	TSharedRef<FBoolKeyArea> KeyArea = MakeShareable( new FBoolKeyArea( BoolSection->GetCurve(), ExternalValue, BoolSection ) );
+	LayoutBuilder.SetSectionAsKeyArea( KeyArea );
 }
 
 int32 FBoolPropertySection::OnPaintSection( FSequencerSectionPainter& Painter ) const
@@ -57,7 +65,7 @@ int32 FBoolPropertySection::OnPaintSection( FSequencerSectionPainter& Painter )
 	{
 		float ThisTime = SectionSwitchTimes[i];
 		
-		const FColor Color = BoolSection->Eval(ThisTime) ? FColor(0, 255, 0, 125) : FColor(255, 0, 0, 125);
+		const FColor Color = BoolSection->Eval(ThisTime, false) ? FColor(0, 255, 0, 125) : FColor(255, 0, 0, 125);
 		
 		FVector2D StartPos(TimeConverter.TimeToPixel(ThisTime), VerticalOffset);
 		FVector2D Size(TimeConverter.TimeToPixel(SectionSwitchTimes[i+1]) - StartPos.X, Height);
@@ -74,14 +82,4 @@ int32 FBoolPropertySection::OnPaintSection( FSequencerSectionPainter& Painter )
 	}
 
 	return LayerId + 1;
-}
-
-void FBoolPropertySection::SetIntermediateValue( FPropertyChangedParams PropertyChangedParams )
-{
-	KeyArea->SetIntermediateValue( PropertyChangedParams.GetPropertyValue<bool>() );
-}
-
-void FBoolPropertySection::ClearIntermediateValue()
-{
-	KeyArea->ClearIntermediateValue();
 }
\ No newline at end of file
diff --git a/Engine/Source/Editor/MovieSceneTools/Private/Sections/BoolPropertySection.h b/Engine/Source/Editor/MovieSceneTools/Private/Sections/BoolPropertySection.h
--- a/Engine/Source/Editor/MovieSceneTools/Private/Sections/BoolPropertySection.h
+++ b/Engine/Source/Editor/MovieSceneTools/Private/Sections/BoolPropertySection.h
@@ -16,9 +16,32 @@ class FBoolPropertySection
 {
 public:
 
-	FBoolPropertySection(UMovieSceneSection& InSectionObject, const FText& SectionName, ISequencer* InSequencer)
-		: FPropertySection(InSectionObject, SectionName)
-		, Sequencer(InSequencer) {}
+	/**
+	* Creates a new bool property section for editing non-property bool sections.
+	*
+	* @param InSectionObject The section object which is being displayed and edited.
+	* @param InDisplayName A display name for the section being displayed and edited.
+	* @TODO Make another section which is a good base class for non-property bool sections.
+	*/
+	FBoolPropertySection(UMovieSceneSection& InSectionObject, const FText& InDisplayName)
+		: FPropertySection(InSectionObject, InDisplayName)
+	{
+	}
+
+	/**
+	* Creates a new bool property section.
+	*
+	* @param InSequencer The sequencer which is controlling this property section.
+	* @param InObjectBinding The object binding for the object which owns the property that this section is animating.
+	* @param InPropertyName The name of the property which is animated by this section.
+	* @param InPropertyPath A string representing the path to the property which is animated by this section.
+	* @param InSectionObject The section object which is being displayed and edited.
+	* @param InDisplayName A display name for the section being displayed and edited.
+	*/
+	FBoolPropertySection(ISequencer* InSequencer, FGuid InObjectBinding, FName InPropertyName, const FString& InPropertyPath, UMovieSceneSection& InSectionObject, const FText& InDisplayName)
+		: FPropertySection(InSequencer, InObjectBinding, InPropertyName, InPropertyPath, InSectionObject, InDisplayName)
+	{
+	}
 
 public:
 
@@ -26,14 +49,4 @@ public:
 
 	virtual void GenerateSectionLayout(class ISectionLayoutBuilder& LayoutBuilder) const override;
 	virtual int32 OnPaintSection(FSequencerSectionPainter& Painter) const override;
-	virtual void SetIntermediateValue(FPropertyChangedParams PropertyChangedParams) override;
-	virtual void ClearIntermediateValue() override;
-
-protected:
-
-	mutable TSharedPtr<FBoolKeyArea> KeyArea;
-
-private:
-
-	ISequencer* Sequencer;
 };
diff --git a/Engine/Source/Editor/MovieSceneTools/Private/Sections/BytePropertySection.cpp b/Engine/Source/Editor/MovieSceneTools/Private/Sections/BytePropertySection.cpp
--- a/Engine/Source/Editor/MovieSceneTools/Private/Sections/BytePropertySection.cpp
+++ b/Engine/Source/Editor/MovieSceneTools/Private/Sections/BytePropertySection.cpp
@@ -10,26 +10,20 @@
 void FBytePropertySection::GenerateSectionLayout( class ISectionLayoutBuilder& LayoutBuilder ) const
 {
 	UMovieSceneByteSection* ByteSection = Cast<UMovieSceneByteSection>( &SectionObject );
+	TAttribute<TOptional<uint8>> ExternalValue;
+	ExternalValue.Bind(TAttribute<TOptional<uint8>>::FGetter::CreateLambda([&]
+	{
+		return GetPropertyValue<uint8>();
+	}));
+
 	if ( Enum != nullptr )
 	{
-		KeyArea = MakeShareable( new FEnumKeyArea( ByteSection->GetCurve(), ByteSection, Enum ) );
-		LayoutBuilder.SetSectionAsKeyArea( KeyArea.ToSharedRef() );
+		TSharedRef<FEnumKeyArea> KeyArea = MakeShareable( new FEnumKeyArea( ByteSection->GetCurve(), ExternalValue, ByteSection, Enum ) );
+		LayoutBuilder.SetSectionAsKeyArea( KeyArea );
 	}
 	else
 	{
-		KeyArea = MakeShareable( new FByteKeyArea( ByteSection->GetCurve(), ByteSection ) );
-		LayoutBuilder.SetSectionAsKeyArea( KeyArea.ToSharedRef() );
+		TSharedRef<FByteKeyArea> KeyArea = MakeShareable( new FByteKeyArea( ByteSection->GetCurve(), ExternalValue, ByteSection ) );
+		LayoutBuilder.SetSectionAsKeyArea( KeyArea );
 	}
-}
-
-
-void FBytePropertySection::SetIntermediateValue( FPropertyChangedParams PropertyChangedParams )
-{
-	KeyArea->SetIntermediateValue( PropertyChangedParams.GetPropertyValue<uint8>() );
-}
-
-
-void FBytePropertySection::ClearIntermediateValue()
-{
-	KeyArea->ClearIntermediateValue();
 }
\ No newline at end of file
diff --git a/Engine/Source/Editor/MovieSceneTools/Private/Sections/BytePropertySection.h b/Engine/Source/Editor/MovieSceneTools/Private/Sections/BytePropertySection.h
--- a/Engine/Source/Editor/MovieSceneTools/Private/Sections/BytePropertySection.h
+++ b/Engine/Source/Editor/MovieSceneTools/Private/Sections/BytePropertySection.h
@@ -5,7 +5,6 @@
 #include "PropertySection.h"
 
 
-class FByteKeyArea;
 class UMovieSceneSection;
 
 
@@ -17,21 +16,31 @@ class FBytePropertySection
 {
 public:
 
-	FBytePropertySection(UMovieSceneSection& InSectionObject, const FText& InSectionName, UEnum* InEnum)
-		: FPropertySection(InSectionObject, InSectionName)
-		, Enum( InEnum )
-	{ }
+	/**
+	* Creates a new byte property section.
+	*
+	* @param InSequencer The sequencer which is controlling this property section.
+	* @param InObjectBinding The object binding for the object which owns the property that this section is animating.
+	* @param InPropertyName The name of the property which is animated by this section.
+	* @param InPropertyPath A string representing the path to the property which is animated by this section.
+	* @param InSectionObject The section object which is being displayed and edited.
+	* @param InDisplayName A display name for the section being displayed and edited.
+	* @param InEnum An optional enum parameter for byte tracks representing enum properties which is used display and populate the enum UI.
+	*/
+	FBytePropertySection(ISequencer* InSequencer, FGuid InObjectBinding, FName InPropertyName, const FString& InPropertyPath, UMovieSceneSection& InSectionObject, const FText& InDisplayName, UEnum* InEnum = nullptr)
+		: FPropertySection(InSequencer, InObjectBinding, InPropertyName, InPropertyPath, InSectionObject, InDisplayName)
+		, Enum(InEnum)
+	{
+	}
 
 public:
 
 	// FPropertySection interface
 
 	virtual void GenerateSectionLayout(class ISectionLayoutBuilder& LayoutBuilder) const override;
-	virtual void SetIntermediateValue(FPropertyChangedParams PropertyChangedParams) override;
-	virtual void ClearIntermediateValue() override;
 
 private:
 
-	mutable TSharedPtr<FByteKeyArea> KeyArea;
+	/** An optional UEnum for byte sections which animate an enum property. */
 	UEnum* Enum;
 };
