Copying //UE4/Dev-Sequencer to //UE4/Main (Source: //UE4/Dev-Sequencer @ 3617642)

diff --git a/Engine/Plugins/Media/MediaPlayerEditor/Source/MediaPlayerEditor/Private/Widgets/SMediaPlayerEditorPlaylist.h b/Engine/Plugins/Media/MediaPlayerEditor/Source/MediaPlayerEditor/Private/Widgets/SMediaPlayerEditorPlaylist.h
--- a/Engine/Plugins/Media/MediaPlayerEditor/Source/MediaPlayerEditor/Private/Widgets/SMediaPlayerEditorPlaylist.h
+++ b/Engine/Plugins/Media/MediaPlayerEditor/Source/MediaPlayerEditor/Private/Widgets/SMediaPlayerEditorPlaylist.h
@@ -13,6 +13,7 @@ class UMediaPlayer;
 struct FMediaSourceTableEntry;
 enum class EMediaEvent;
 
+
 /**
  * Implements the play list of the MediaPlayer asset editor.
  */
@@ -56,11 +57,8 @@ private:
 	/** Callback for media player events. */
 	void HandleMediaPlayerMediaEvent(EMediaEvent Event);
 
-	/** Callback for double-clicking an entry in the media source list. */
-	void HandleMediaSourceListDoubleClick(TSharedPtr<FMediaSourceTableEntry> InItem);
-
-	/** Callback for generating a row widget for the media source list view. */
-	TSharedRef<ITableRow> HandleMediaSourceListGenerateRow(TSharedPtr<FMediaSourceTableEntry> Entry, const TSharedRef<STableViewBase>& OwnerTable);
+	/** Callback for clicking the Save button. */
+	FReply HandleSavePlaylistButtonClicked();
 
 private:
 
diff --git a/Engine/Plugins/Media/MediaPlayerEditor/Source/MediaPlayerEditor/Private/Widgets/SMediaPlayerEditorStats.cpp b/Engine/Plugins/Media/MediaPlayerEditor/Source/MediaPlayerEditor/Private/Widgets/SMediaPlayerEditorStats.cpp
--- a/Engine/Plugins/Media/MediaPlayerEditor/Source/MediaPlayerEditor/Private/Widgets/SMediaPlayerEditorStats.cpp
+++ b/Engine/Plugins/Media/MediaPlayerEditor/Source/MediaPlayerEditor/Private/Widgets/SMediaPlayerEditorStats.cpp
@@ -1,14 +1,15 @@
 // Copyright 1998-2017 Epic Games, Inc. All Rights Reserved.
 
-#include "Widgets/SMediaPlayerEditorStats.h"
-#include "IMediaPlayer.h"
+#include "SMediaPlayerEditorStats.h"
+
+#include "EditorStyleSet.h"
 #include "MediaPlayer.h"
-#include "Widgets/Layout/SBorder.h"
-#include "Widgets/Text/STextBlock.h"
+#include "MediaPlayerFacade.h"
 #include "Widgets/SBoxPanel.h"
-#include "Widgets/Layout/SScrollBox.h"
 #include "Widgets/Input/SButton.h"
-#include "EditorStyleSet.h"
+#include "Widgets/Layout/SBorder.h"
+#include "Widgets/Layout/SScrollBox.h"
+#include "Widgets/Text/STextBlock.h"
 #include "HAL/PlatformApplicationMisc.h"
 
 #define LOCTEXT_NAMESPACE "SMediaPlayerEditorStats"
@@ -64,14 +65,12 @@ void SMediaPlayerEditorStats::Construct(const FArguments& InArgs, UMediaPlayer&
 
 FText SMediaPlayerEditorStats::HandleStatsTextBlockText() const
 {
-	FMediaPlayerBase& Player = MediaPlayer->GetBasePlayer();
-
-	if (Player.GetUrl().IsEmpty())
+	if (MediaPlayer->GetUrl().IsEmpty())
 	{
 		return LOCTEXT("NoMediaOpened", "No media opened");
 	}
 
-	return FText::FromString(Player.GetStats());
+	return FText::FromString(MediaPlayer->GetPlayerFacade()->GetStats());
 }
 
 
diff --git a/Engine/Plugins/Media/MediaPlayerEditor/Source/MediaPlayerEditor/Private/Widgets/SMediaPlayerEditorViewer.cpp b/Engine/Plugins/Media/MediaPlayerEditor/Source/MediaPlayerEditor/Private/Widgets/SMediaPlayerEditorViewer.cpp
--- a/Engine/Plugins/Media/MediaPlayerEditor/Source/MediaPlayerEditor/Private/Widgets/SMediaPlayerEditorViewer.cpp
+++ b/Engine/Plugins/Media/MediaPlayerEditor/Source/MediaPlayerEditor/Private/Widgets/SMediaPlayerEditorViewer.cpp
@@ -1,36 +1,40 @@
 // Copyright 1998-2017 Epic Games, Inc. All Rights Reserved.
 
 #include "Widgets/SMediaPlayerEditorViewer.h"
-#include "Styling/CoreStyle.h"
-#include "Widgets/Input/SSlider.h"
-#include "IMediaPlayer.h"
-#include "MediaPlayer.h"
+
+#include "Internationalization/Culture.h"
 #include "Modules/ModuleManager.h"
+
+#include "IMediaModule.h"
+#include "IMediaEventSink.h"
+#include "IMediaPlayerFactory.h"
+#include "MediaPlayer.h"
+
+#include "EditorStyleSet.h"
+#include "Framework/Application/SlateApplication.h"
 #include "Framework/Commands/UIAction.h"
+#include "Framework/MultiBox/MultiBoxBuilder.h"
+#include "Framework/Notifications/NotificationManager.h"
+#include "Styling/CoreStyle.h"
 #include "Textures/SlateIcon.h"
-#include "Internationalization/Culture.h"
 #include "Widgets/SBoxPanel.h"
 #include "Widgets/SOverlay.h"
-#include "Widgets/Layout/SBorder.h"
-#include "Widgets/Layout/SSpacer.h"
 #include "Widgets/Images/SImage.h"
-#include "Widgets/Notifications/SProgressBar.h"
-#include "Widgets/Text/STextBlock.h"
-#include "Widgets/Layout/SBox.h"
-#include "Framework/MultiBox/MultiBoxBuilder.h"
-#include "Widgets/Input/SEditableTextBox.h"
+#include "Widgets/Images/SThrobber.h"
 #include "Widgets/Input/SButton.h"
 #include "Widgets/Input/SComboButton.h"
-#include "EditorStyleSet.h"
-#include "IMediaModule.h"
-#include "IMediaPlayerFactory.h"
-#include "Framework/Notifications/NotificationManager.h"
-#include "Widgets/Layout/SScaleBox.h"
+#include "Widgets/Input/SEditableTextBox.h"
+#include "Widgets/Input/SSlider.h"
+#include "Widgets/Layout/SBorder.h"
+#include "Widgets/Layout/SBox.h"
+#include "Widgets/Layout/SSpacer.h"
 #include "Widgets/Notifications/SNotificationList.h"
-#include "Shared/MediaPlayerEditorSettings.h"
-#include "Widgets/SMediaPlayerEditorOutput.h"
-#include "Framework/Application/SlateApplication.h"
-#include "Widgets/SMediaPlayerEditorOverlay.h"
+#include "Widgets/Notifications/SProgressBar.h"
+#include "Widgets/Text/STextBlock.h"
+
+#include "MediaPlayerEditorSettings.h"
+#include "SMediaPlayerEditorCache.h"
+#include "SMediaPlayerEditorViewport.h"
 
 
 #define LOCTEXT_NAMESPACE "SMediaPlayerEditorViewer"
@@ -79,6 +83,30 @@ void SMediaPlayerEditorViewer::Construct(const FArguments& InArgs, UMediaPlayer&
 		}
 	}
 
+	// initialize capture source menu
+	FMenuBuilder SourceMenuBuilder(true, nullptr);
+	{
+		SourceMenuBuilder.BeginSection("CaptureDevicesSection", LOCTEXT("CaptureDevicesSection", "Capture Devices"));
+		{
+			SourceMenuBuilder.AddSubMenu(
+				LOCTEXT("AudioMenuLabel", "Audio"),
+				LOCTEXT("AudioMenuTooltip", "Available audio capture devices"),
+				FNewMenuDelegate::CreateRaw(this, &SMediaPlayerEditorViewer::HandleAudioCaptureDevicesMenuNewMenu),
+				false,
+				FSlateIcon()
+			);
+
+			SourceMenuBuilder.AddSubMenu(
+				LOCTEXT("VideoMenuLabel", "Video"),
+				LOCTEXT("VideoMenuTooltip", "Available video capture devices"),
+				FNewMenuDelegate::CreateRaw(this, &SMediaPlayerEditorViewer::HandleVideoCaptureDevicesMenuNewMenu),
+				false,
+				FSlateIcon()
+			);
+		}
+		SourceMenuBuilder.EndSection();
+	}
+
 	// initialize settings menu
 	FMenuBuilder SettingsMenuBuilder(true, nullptr);
 	{
@@ -87,7 +115,7 @@ void SMediaPlayerEditorViewer::Construct(const FArguments& InArgs, UMediaPlayer&
 			SettingsMenuBuilder.AddSubMenu(
 				LOCTEXT("DecoderMenuLabel", "Decoder"),
 				LOCTEXT("DecoderMenuTooltip", "Select the desired media decoder"),
-				FNewMenuDelegate::CreateRaw(this, &SMediaPlayerEditorViewer::HandlePlayerMenuNewMenu),
+				FNewMenuDelegate::CreateRaw(this, &SMediaPlayerEditorViewer::HandleDecoderMenuNewMenu),
 				false,
 				FSlateIcon()
 			);
@@ -138,6 +166,50 @@ void SMediaPlayerEditorViewer::Construct(const FArguments& InArgs, UMediaPlayer&
 		}
 		SettingsMenuBuilder.EndSection();
 
+		SettingsMenuBuilder.BeginSection("FormatsSection", LOCTEXT("FormatsSection", "Formats"));
+		{
+			SettingsMenuBuilder.AddSubMenu(
+				LOCTEXT("AudioFormatMenuLabel", "Audio"),
+				LOCTEXT("AudioFormatMenuTooltip", "Select the active audio format"),
+				FNewMenuDelegate::CreateRaw(this, &SMediaPlayerEditorViewer::HandleFormatMenuNewMenu, EMediaPlayerTrack::Audio),
+				false,
+				FSlateIcon()
+			);
+
+			SettingsMenuBuilder.AddSubMenu(
+				LOCTEXT("CaptionFormatMenuLabel", "Captions"),
+				LOCTEXT("CaptionFormatMenuTooltip", "Select the active closed caption format"),
+				FNewMenuDelegate::CreateRaw(this, &SMediaPlayerEditorViewer::HandleFormatMenuNewMenu, EMediaPlayerTrack::Caption),
+				false,
+				FSlateIcon()
+			);
+
+			SettingsMenuBuilder.AddSubMenu(
+				LOCTEXT("SubtitleFormatMenuLabel", "Subtitles"),
+				LOCTEXT("SubtitleFormatMenuTooltip", "Select the active subtitle format"),
+				FNewMenuDelegate::CreateRaw(this, &SMediaPlayerEditorViewer::HandleFormatMenuNewMenu, EMediaPlayerTrack::Subtitle),
+				false,
+				FSlateIcon()
+			);
+
+			SettingsMenuBuilder.AddSubMenu(
+				LOCTEXT("TextFormatMenuLabel", "Text"),
+				LOCTEXT("TextFormatMenuTooltip", "Select the active generic text format"),
+				FNewMenuDelegate::CreateRaw(this, &SMediaPlayerEditorViewer::HandleFormatMenuNewMenu, EMediaPlayerTrack::Text),
+				false,
+				FSlateIcon()
+			);
+
+			SettingsMenuBuilder.AddSubMenu(
+				LOCTEXT("VideoFormatMenuLabel", "Video"),
+				LOCTEXT("VideoFormatMenuTooltip", "Select the active video format"),
+				FNewMenuDelegate::CreateRaw(this, &SMediaPlayerEditorViewer::HandleFormatMenuNewMenu, EMediaPlayerTrack::Video),
+				false,
+				FSlateIcon()
+			);
+		}
+		SettingsMenuBuilder.EndSection();
+
 		SettingsMenuBuilder.BeginSection("ViewSection", LOCTEXT("ViewSection", "View"));
 		{
 			SettingsMenuBuilder.AddSubMenu(
@@ -168,6 +240,7 @@ void SMediaPlayerEditorViewer::Construct(const FArguments& InArgs, UMediaPlayer&
 		SettingsMenuBuilder.EndSection();
 	}
 
+	// widget contents
 	ChildSlot
 	[
 		SNew(SOverlay)
@@ -179,21 +252,42 @@ void SMediaPlayerEditorViewer::Construct(const FArguments& InArgs, UMediaPlayer&
 				+ SVerticalBox::Slot()
 					.AutoHeight()
 					[
+						// url area
 						SNew(SHorizontalBox)
 
+						+ SHorizontalBox::Slot()
+							.AutoWidth()
+							[
+								// capture source drop-down
+								SNew(SComboButton)
+									.ContentPadding(0.0f)
+									.ButtonContent()
+									[
+										SNew(SImage)
+											.Image(InStyle->GetBrush("MediaPlayerEditor.SourceButton"))
+									]
+									.ButtonStyle(FEditorStyle::Get(), "ToggleButton")
+									.ForegroundColor(FSlateColor::UseForeground())
+									.MenuContent()
+									[
+										SourceMenuBuilder.MakeWidget()
+									]
+							]
+
 						+ SHorizontalBox::Slot()
 							.FillWidth(1.0f)
+							.Padding(4.0f, 0.0f, 0.0f, 0.0f)
 							[
 								// url box
 								SAssignNew(UrlTextBox, SEditableTextBox)
 									.BackgroundColor_Lambda([this]() -> FSlateColor {
-										return (!MediaPlayer->IsPreparing() && (MediaPlayer->GetUrl() != MediaPlayer->GetLastUrl()))
-											? FLinearColor::Red
-											: FLinearColor::White;
+										return (MediaPlayer->IsPreparing() || LastUrl.IsEmpty()) ? FLinearColor::White : FLinearColor::Red;
 									})
 									.ClearKeyboardFocusOnCommit(true)
 									.HintText(LOCTEXT("UrlTextBoxHint", "Media URL"))
-									.Text_Lambda([this]() -> FText { return FText::FromString(MediaPlayer->GetLastUrl()); })
+									.Text_Lambda([this]() -> FText {
+										return LastUrl.IsEmpty() ? FText::FromString(MediaPlayer->GetUrl()) : LastUrl;
+									})
 									.ToolTipText(LOCTEXT("UrlTextBoxToolTip", "Enter the URL of a media source"))
 									.OnKeyDownHandler(this, &SMediaPlayerEditorViewer::HandleUrlBoxKeyDown)
 									.OnTextCommitted_Lambda([this](const FText& InText, ETextCommit::Type InCommitType){
@@ -240,72 +334,9 @@ void SMediaPlayerEditorViewer::Construct(const FArguments& InArgs, UMediaPlayer&
 						SNew(SBorder)
 							.BorderImage(FCoreStyle::Get().GetBrush("BlackBrush"))
 							.Padding(0.0f)
-							.OnMouseButtonUp(this, &SMediaPlayerEditorViewer::HandleTextureMouseButtonUp)
 							[
-								SNew(SOverlay)
-
-								+ SOverlay::Slot()
-									[
-										SNew(SScaleBox)
-											.Stretch_Lambda([]() -> EStretch::Type {
-												auto Settings = GetDefault<UMediaPlayerEditorSettings>();
-												if (Settings->ViewportScale == EMediaPlayerEditorScale::Fill)
-												{
-													return EStretch::Fill;
-												}
-												if (Settings->ViewportScale == EMediaPlayerEditorScale::Fit)
-												{
-													return EStretch::ScaleToFit;
-												}
-												return EStretch::None;
-											})
-											[
-												SNew(SOverlay)
-
-												+ SOverlay::Slot()
-													[
-														// movie texture
-														SNew(SMediaPlayerEditorOutput, InMediaPlayer)
-													]
-
-												+ SOverlay::Slot()
-													[
-														// text overlays
-														SNew(SMediaPlayerEditorOverlay, InMediaPlayer)
-															.Visibility_Lambda([]{ return GetDefault<UMediaPlayerEditorSettings>()->ShowTextOverlays ? EVisibility::Visible : EVisibility::Collapsed; })
-													]
-											]
-									]
-
-								+ SOverlay::Slot()
-									.Padding(FMargin(12.0f, 8.0f))
-									[
-										SNew(SHorizontalBox)
-
-										+ SHorizontalBox::Slot()
-											.HAlign(HAlign_Left)
-											.VAlign(VAlign_Top)
-											[
-												// playback state
-												SNew(STextBlock)
-													.ColorAndOpacity(FSlateColor::UseSubduedForeground())
-													.Font(FSlateFontInfo(FPaths::EngineContentDir() / TEXT("Slate/Fonts/Roboto-Regular.ttf"), 18))
-													.ShadowOffset(FVector2D(1.f, 1.f))
-													.Text(this, &SMediaPlayerEditorViewer::HandleOverlayStateText)
-											]
-
-										+ SHorizontalBox::Slot()
-											.HAlign(HAlign_Right)
-											.VAlign(VAlign_Top)
-											[
-												// player name
-												SNew(STextBlock)
-													.ColorAndOpacity(FSlateColor::UseSubduedForeground())
-													.Font(FSlateFontInfo(FPaths::EngineContentDir() / TEXT("Slate/Fonts/Roboto-Regular.ttf"), 18))
-													.ShadowOffset(FVector2D(1.f, 1.f))
-													.Text(this, &SMediaPlayerEditorViewer::HandleOverlayPlayerNameText)
-											]
-									]
+								// movie area
+								SNew(SMediaPlayerEditorViewport, InMediaPlayer, InStyle)
 							]
 					]
 
@@ -314,6 +345,7 @@ void SMediaPlayerEditorViewer::Construct(const FArguments& InArgs, UMediaPlayer&
 					.Padding(0.0f, 2.0f, 0.0f, 0.0f)
 					.VAlign(VAlign_Center)
 					[
+						// playback controls
 						SNew(SBorder)
 							.BorderImage(FEditorStyle::GetBrush("ToolPanel.GroupBorder"))
 							.ForegroundColor(FLinearColor::Gray)
@@ -327,22 +359,79 @@ void SMediaPlayerEditorViewer::Construct(const FArguments& InArgs, UMediaPlayer&
 										SNew(SOverlay)
 
 										+ SOverlay::Slot()
+											.Padding(2.0f, 0.0f, 2.0f, 0.0f)
+											.VAlign(VAlign_Top)
+											[
+												// cache visualization
+												SNew(SMediaPlayerEditorCache, InMediaPlayer, InStyle)
+											]												
+
+										+ SOverlay::Slot()
+											.VAlign(VAlign_Top)
 											[
-												// scrubber
+												// time slider
 												SAssignNew(ScrubberSlider, SSlider)
-													.IsEnabled(this, &SMediaPlayerEditorViewer::HandlePositionSliderIsEnabled)
-													.OnMouseCaptureBegin(this, &SMediaPlayerEditorViewer::HandlePositionSliderMouseCaptureBegin)
-													.OnMouseCaptureEnd(this, &SMediaPlayerEditorViewer::HandlePositionSliderMouseCaptureEnd)
-													.OnValueChanged(this, &SMediaPlayerEditorViewer::HandlePositionSliderValueChanged)
+													.IsEnabled_Lambda([this]() -> bool
+													{
+														return MediaPlayer->SupportsSeeking();
+													})
+													.OnMouseCaptureBegin_Lambda([this]()
+													{
+														ScrubValue = FTimespan::Ratio(MediaPlayer->GetTime(), MediaPlayer->GetDuration());
+
+														if (MediaPlayer->SupportsScrubbing())
+														{
+															PreScrubRate = MediaPlayer->GetRate();
+															MediaPlayer->SetRate(0.0f);
+														}
+													})
+													.OnMouseCaptureEnd_Lambda([this]()
+													{
+														if (MediaPlayer->SupportsScrubbing())
+														{
+															MediaPlayer->SetRate(PreScrubRate);
+														}
+													})
+													.OnValueChanged_Lambda([this](float NewValue)
+													{
+														ScrubValue = NewValue;
+
+														if (!ScrubberSlider->HasMouseCapture() || MediaPlayer->SupportsScrubbing())
+														{
+															MediaPlayer->Seek(MediaPlayer->GetDuration() * NewValue);
+														}
+													})
 													.Orientation(Orient_Horizontal)
-													.Value(this, &SMediaPlayerEditorViewer::HandlePositionSliderValue)
+													.SliderBarColor(FLinearColor::Transparent)
+													.Style(&InStyle->GetWidgetStyle<FSliderStyle>("MediaPlayerEditor.Scrubber"))
+													.Value_Lambda([this]() -> float
+													{
+														if (ScrubberSlider->HasMouseCapture())
+														{
+															return ScrubValue;
+														}
+
+														return (MediaPlayer->GetDuration() > FTimespan::Zero())
+															? FTimespan::Ratio(MediaPlayer->GetTime(), MediaPlayer->GetDuration())
+															: 0.0f;
+													})
+													.Visibility_Lambda([this]() -> EVisibility {
+														return (MediaPlayer->SupportsScrubbing() || MediaPlayer->SupportsSeeking())
+															? EVisibility::Visible
+															: EVisibility::Hidden;
+													})
 											]
 
 										+ SOverlay::Slot()
+											.VAlign(VAlign_Center)
 											[
+												// animated progress bar
 												SNew(SProgressBar)
-													.ToolTipText(LOCTEXT("BufferingTooltip", "Buffering..."))
-													.Visibility_Lambda([this]() -> EVisibility { return MediaPlayer->IsPreparing() ? EVisibility::Visible : EVisibility::Hidden; })
+													.ToolTipText(LOCTEXT("PreparingTooltip", "Preparing..."))
+													.Visibility_Lambda([this]() -> EVisibility
+													{
+														return MediaPlayer->IsPreparing() ? EVisibility::Visible : EVisibility::Hidden;
+													})
 											]
 									]
 
@@ -355,7 +444,7 @@ void SMediaPlayerEditorViewer::Construct(const FArguments& InArgs, UMediaPlayer&
 										// timer
 										+ SHorizontalBox::Slot()
 											.AutoWidth()
-											.Padding(0.0f, 0.0f, 4.0f, 0.0f)
+											.Padding(4.0f, 0.0f, 0.0f, 0.0f)
 											.VAlign(VAlign_Center)
 											[
 												SNew(STextBlock)
@@ -363,20 +452,43 @@ void SMediaPlayerEditorViewer::Construct(const FArguments& InArgs, UMediaPlayer&
 													.ToolTipText(this, &SMediaPlayerEditorViewer::HandleTimerTextBlockToolTipText)
 											]
 
+										// fps
+										+ SHorizontalBox::Slot()
+											.AutoWidth()
+											.Padding(8.0f, 0.0f, 4.0f, 0.0f)
+											.VAlign(VAlign_Center)
+											[
+												SNew(STextBlock)
+													.Text(this, &SMediaPlayerEditorViewer::HandleFpsTextBlockText)
+											]
+
+										// buffering indicator
+										+ SHorizontalBox::Slot()
+											.AutoWidth()
+											.Padding(4.0f, 0.0f)
+											.VAlign(VAlign_Center)
+											[
+												SNew(SThrobber)
+													.ToolTipText(LOCTEXT("BufferingTooltip", "Buffering..."))
+													.Visibility_Lambda([this]() -> EVisibility
+													{
+														return MediaPlayer->IsBuffering() ? EVisibility::Visible : EVisibility::Hidden;
+													})
+											]
+
 										// transport controls
 										+ SHorizontalBox::Slot()
 											.FillWidth(1.0f)
 											.Padding(8.0f, 0.0f)
 											.VAlign(VAlign_Center)
 											[
-												SNew(SSpacer)
-												//SNew(STransportControl)
+												SNew(SSpacer)//SNew(STransportControl)
 											]
 
 										// settings
 										+ SHorizontalBox::Slot()
 											.AutoWidth()
-											.Padding(4.0f, 0.0f, 0.0f, 0.0f)
+											.Padding(4.0f, 0.0f, 4.0f, 0.0f)
 											.VAlign(VAlign_Center)
 											[
 												SNew(SComboButton)
@@ -416,6 +528,7 @@ void SMediaPlayerEditorViewer::Construct(const FArguments& InArgs, UMediaPlayer&
 
 		+ SOverlay::Slot()
 			[
+				// drag & drop indicator
 				SNew(SBorder)
 					.BorderImage(FEditorStyle::GetBrush("WhiteBrush"))
 					.BorderBackgroundColor_Lambda([this]() -> FLinearColor {
@@ -430,8 +543,6 @@ void SMediaPlayerEditorViewer::Construct(const FArguments& InArgs, UMediaPlayer&
 					})
 			]
 	];
-
-	RegisterActiveTimer(0.0f, FWidgetActiveTimerDelegate::CreateSP(this, &SMediaPlayerEditorViewer::HandleActiveTimer));
 }
 
 
@@ -505,8 +616,28 @@ FReply SMediaPlayerEditorViewer::OnDrop(const FGeometry& MyGeometry, const FDrag
 /* SMediaPlayerEditorPlayer implementation
  *****************************************************************************/
 
+void SMediaPlayerEditorViewer::MakeCaptureDeviceMenu(TArray<FMediaCaptureDeviceInfo>& DeviceInfos, FMenuBuilder& MenuBuilder)
+{
+	for (const auto& DeviceInfo : DeviceInfos)
+	{
+		MenuBuilder.AddMenuEntry(
+			DeviceInfo.DisplayName,
+			FText::GetEmpty(),
+			FSlateIcon(),
+			FUIAction(FExecuteAction::CreateLambda([this, Url = DeviceInfo.Url] {
+				MediaPlayer->OpenUrl(Url);
+			})),
+			NAME_None,
+			EUserInterfaceActionType::Button
+		);
+	}
+}
+
+
 void SMediaPlayerEditorViewer::OpenUrl(const FText& TextUrl)
 {
+	LastUrl = TextUrl;
+
 	FString Url = TextUrl.ToString();
 	{
 		Url.Trim();
@@ -545,118 +676,15 @@ void SMediaPlayerEditorViewer::SetDesiredPlayerName(FName PlayerName)
 /* SMediaPlayerEditorPlayer callbacks
  *****************************************************************************/
 
-EActiveTimerReturnType SMediaPlayerEditorViewer::HandleActiveTimer(double InCurrentTime, float InDeltaTime)
-{
-	return EActiveTimerReturnType::Continue;
-}
-
-
-void SMediaPlayerEditorViewer::HandleMediaPlayerMediaEvent(EMediaEvent Event)
-{
-	if (Event == EMediaEvent::MediaOpenFailed)
-	{
-		FNotificationInfo NotificationInfo(LOCTEXT("MediaOpenFailedError", "The media failed to open. Check Output Log for details!"));
-		{
-			NotificationInfo.ExpireDuration = 2.0f;
-		}
-
-		FSlateNotificationManager::Get().AddNotification(NotificationInfo)->SetCompletionState(SNotificationItem::CS_Fail);
-	}
-}
-
-
-FText SMediaPlayerEditorViewer::HandleOverlayPlayerNameText() const
-{
-	if (MediaPlayer->IsPlaying())
-	{
-		return FText::GetEmpty();
-	}
-
-	FName PlayerName = MediaPlayer->GetPlayerName();
-
-	if ((PlayerName == NAME_None) || MediaPlayer->GetUrl().IsEmpty())
-	{
-		const FName DesiredPlayerName = MediaPlayer->GetDesiredPlayerName();
-
-		if (DesiredPlayerName == NAME_None)
-		{
-			return LOCTEXT("AutoPlayerName", "Auto");
-		}
-
-		return FText::FromName(DesiredPlayerName);
-	}
-
-	return FText::FromName(PlayerName);
-}
-
-
-FText SMediaPlayerEditorViewer::HandleOverlayStateText() const
+void SMediaPlayerEditorViewer::HandleAudioCaptureDevicesMenuNewMenu(FMenuBuilder& MenuBuilder)
 {
-	if (MediaPlayer->GetUrl().IsEmpty())
-	{
-		return LOCTEXT("StateOverlayNoMedia", "No Media");
-	}
-
-	if (!MediaPlayer->IsReady())
-	{
-		return LOCTEXT("StateOverlayStopped", "Not Ready");
-	}
-
-	if (MediaPlayer->IsPaused())
-	{
-		return LOCTEXT("StateOverlayPaused", "Paused");
-	}
-
-	if (!MediaPlayer->IsPlaying())
-	{
-		return FText::FromString(FPaths::GetBaseFilename(MediaPlayer->GetUrl()));
-	}
-
-	float Rate = MediaPlayer->GetRate();
-
-	if (FMath::IsNearlyZero(Rate) || (Rate == 1.0f))
-	{
-		return FText::GetEmpty();
-	}
-
-	if (Rate < 0.0f)
-	{
-		return FText::Format(LOCTEXT("StateOverlayReverseFormat", "Reverse ({0}x)"), FText::AsNumber(-Rate));
-	}
-
-	return FText::Format(LOCTEXT("StateOverlayForwardFormat", "Forward ({0}x)"), FText::AsNumber(Rate));
-}
-
-
-TOptional<float> SMediaPlayerEditorViewer::HandlePlaybackRateBoxMaxValue() const
-{
-	FFloatRange Rates = MediaPlayer->GetForwardRates(false);
-
-	return Rates.IsEmpty() ? 0.0f : Rates.GetUpperBoundValue();
-}
-
-
-TOptional<float> SMediaPlayerEditorViewer::HandlePlaybackRateBoxMinValue() const
-{
-	FFloatRange Rates = MediaPlayer->GetReverseRates(false);
-
-	return Rates.IsEmpty() ? 0.0f : Rates.GetLowerBoundValue();
+	TArray<FMediaCaptureDeviceInfo> DeviceInfos;
+	MediaCaptureSupport::EnumerateAudioCaptureDevices(DeviceInfos);
+	MakeCaptureDeviceMenu(DeviceInfos, MenuBuilder);
 }
 
 
-TOptional<float> SMediaPlayerEditorViewer::HandlePlaybackRateSpinBoxValue() const
-{
-	return MediaPlayer->GetRate();
-}
-
-
-void SMediaPlayerEditorViewer::HandlePlaybackRateBoxValueChanged( float NewValue )
-{
-	MediaPlayer->SetRate(NewValue);
-}
-
-
-void SMediaPlayerEditorViewer::HandlePlayerMenuNewMenu(FMenuBuilder& MenuBuilder)
+void SMediaPlayerEditorViewer::HandleDecoderMenuNewMenu(FMenuBuilder& MenuBuilder)
 {
 	// automatic player option
 	MenuBuilder.AddMenuEntry(
@@ -712,9 +740,9 @@ void SMediaPlayerEditorViewer::HandlePlayerMenuNewMenu(FMenuBuilder& MenuBuilder
 			FText::FromString(FString::Join(Factory->GetSupportedPlatforms(), TEXT(", "))),
 			FSlateIcon(),
 			FUIAction(
-				FExecuteAction::CreateLambda([=]{ SetDesiredPlayerName(PlayerName); }),
-				FCanExecuteAction::CreateLambda([=]{ return SupportsRunningPlatform; }),
-				FIsActionChecked::CreateLambda([=]{ return MediaPlayer->GetDesiredPlayerName() == PlayerName; })
+				FExecuteAction::CreateLambda([=] { SetDesiredPlayerName(PlayerName); }),
+				FCanExecuteAction::CreateLambda([=] { return SupportsRunningPlatform; }),
+				FIsActionChecked::CreateLambda([=] { return MediaPlayer->GetDesiredPlayerName() == PlayerName; })
 			),
 			NAME_None,
 			EUserInterfaceActionType::RadioButton
@@ -723,47 +751,132 @@ void SMediaPlayerEditorViewer::HandlePlayerMenuNewMenu(FMenuBuilder& MenuBuilder
 }
 
 
-bool SMediaPlayerEditorViewer::HandlePositionSliderIsEnabled() const
+void SMediaPlayerEditorViewer::HandleFormatMenuNewMenu(FMenuBuilder& MenuBuilder, EMediaPlayerTrack TrackType)
 {
-	return MediaPlayer->SupportsSeeking();
-}
+	const int32 SelectedTrack = MediaPlayer->GetSelectedTrack(TrackType);
+	const int32 NumFormats = MediaPlayer->GetNumTrackFormats(TrackType, SelectedTrack);
 
+	if ((SelectedTrack != INDEX_NONE) && (NumFormats > 0))
+	{
+		for (int32 FormatIndex = 0; FormatIndex < NumFormats; ++FormatIndex)
+		{
+			FText DisplayText;
 
-void SMediaPlayerEditorViewer::HandlePositionSliderMouseCaptureBegin()
-{
-	if (MediaPlayer->SupportsScrubbing())
+			if (TrackType == EMediaPlayerTrack::Audio)
+			{
+				const uint32 Channels = MediaPlayer->GetAudioTrackChannels(SelectedTrack, FormatIndex);
+				const uint32 SampleRate = MediaPlayer->GetAudioTrackSampleRate(SelectedTrack, FormatIndex);
+				const FString Type = MediaPlayer->GetAudioTrackType(SelectedTrack, FormatIndex);
+
+				DisplayText = FText::Format(LOCTEXT("TrackFormatMenuAudioFormat", "{0}: {1} {2} channels @ {3} Hz"),
+					FText::AsNumber(FormatIndex),
+					FText::FromString(Type),
+					FText::AsNumber(Channels),
+					FText::AsNumber(SampleRate)
+				);
+			}
+			else if (TrackType == EMediaPlayerTrack::Video)
+			{
+				const FIntPoint Dim = MediaPlayer->GetVideoTrackDimensions(SelectedTrack, FormatIndex);
+				const float FrameRate = MediaPlayer->GetVideoTrackFrameRate(SelectedTrack, FormatIndex);
+				const TRange<float> FrameRates = MediaPlayer->GetVideoTrackFrameRates(SelectedTrack, FormatIndex);
+				const FString Type = MediaPlayer->GetVideoTrackType(SelectedTrack, FormatIndex);
+
+				FFormatNamedArguments Arguments;
+				Arguments.Add(TEXT("Index"), FText::AsNumber(FormatIndex));
+				Arguments.Add(TEXT("DimX"), FText::AsNumber(Dim.X));
+				Arguments.Add(TEXT("DimY"), FText::AsNumber(Dim.Y));
+				Arguments.Add(TEXT("Fps"), FText::AsNumber(FrameRate));
+				Arguments.Add(TEXT("Type"), FText::FromString(Type));
+
+				if (FrameRates.IsDegenerate() && (FrameRates.GetLowerBoundValue() == FrameRate))
+				{
+					DisplayText = FText::Format(LOCTEXT("TrackFormatMenuVideoFormat", "{Index}: {Type} {DimX}x{DimY} {Fps} fps"), Arguments);
+				}
+				else
+				{
+					Arguments.Add(TEXT("FpsLower"), FText::AsNumber(FrameRates.GetLowerBoundValue()));
+					Arguments.Add(TEXT("FpsUpper"), FText::AsNumber(FrameRates.GetUpperBoundValue()));
+
+					DisplayText = FText::Format(LOCTEXT("TrackFormatMenuVideoFormat", "{Index}: {Type} {DimX}x{DimY} {Fps} [{FpsLower}-{FpsUpper}] fps"), Arguments);
+				}
+			}
+			else
+			{
+				DisplayText = LOCTEXT("TrackFormatDefault", "Default");
+			}
+
+			MenuBuilder.AddMenuEntry(
+				DisplayText,
+				FText::GetEmpty(),
+				FSlateIcon(),
+				FUIAction(
+					FExecuteAction::CreateLambda([=] { MediaPlayer->SetTrackFormat(TrackType, SelectedTrack, FormatIndex); }),
+					FCanExecuteAction(),
+					FIsActionChecked::CreateLambda([=] { return (MediaPlayer->GetTrackFormat(TrackType, SelectedTrack) == FormatIndex); })
+				),
+				NAME_None,
+				EUserInterfaceActionType::RadioButton
+			);
+		}
+	}
+	else
 	{
-		PreScrubRate = MediaPlayer->GetRate();
-		MediaPlayer->SetRate(0.0f);
+		FText Label = (SelectedTrack == INDEX_NONE)
+			? LOCTEXT("NoTrackSelectedLabel", "No track selected")
+			: LOCTEXT("NoFormatsAvailableLabel", "No formats available");
+
+		TSharedRef<SWidget> NoTracksAvailableWidget = SNew(SBox)
+			.Padding(FMargin(2.0f, 2.0f, 2.0f, 0.0f))
+			[
+				SNew(STextBlock)
+				.ColorAndOpacity(FSlateColor::UseSubduedForeground())
+			.Text(Label)
+			];
+
+		MenuBuilder.AddWidget(NoTracksAvailableWidget, FText::GetEmpty(), true, false);
 	}
 }
 
 
-void SMediaPlayerEditorViewer::HandlePositionSliderMouseCaptureEnd()
+FText SMediaPlayerEditorViewer::HandleFpsTextBlockText() const
 {
-	if (MediaPlayer->SupportsScrubbing())
+	if (!MediaPlayer->IsReady())
 	{
-		MediaPlayer->SetRate(PreScrubRate);
+		return FText::GetEmpty();
 	}
-}
 
+	const int32 SelectedTrack = MediaPlayer->GetSelectedTrack(EMediaPlayerTrack::Video);
+	const int32 SelectedFormat = MediaPlayer->GetTrackFormat(EMediaPlayerTrack::Video, SelectedTrack);
+	const float FrameRate = MediaPlayer->GetVideoTrackFrameRate(SelectedTrack, SelectedFormat);
 
-float SMediaPlayerEditorViewer::HandlePositionSliderValue() const
-{
-	if (MediaPlayer->GetDuration() > FTimespan::Zero())
+	// empty string if fps n/a
+	if (FrameRate <= 0.0f)
 	{
-		return (float)MediaPlayer->GetTime().GetTicks() / (float)MediaPlayer->GetDuration().GetTicks();
+		return FText::GetEmpty();
 	}
 
-	return 0.0;
+	FNumberFormattingOptions FormattingOptions;
+	FormattingOptions.MaximumFractionalDigits = 3;
+
+	return FText::Format(LOCTEXT("FpsTextBlockFormat", "{0} fps"), FText::AsNumber(FrameRate, &FormattingOptions));
 }
 
 
-void SMediaPlayerEditorViewer::HandlePositionSliderValueChanged( float NewValue )
+void SMediaPlayerEditorViewer::HandleMediaPlayerMediaEvent(EMediaEvent Event)
 {
-	if (!ScrubberSlider->HasMouseCapture() || MediaPlayer->SupportsScrubbing())
+	if (Event == EMediaEvent::MediaOpened)
 	{
-		MediaPlayer->Seek(MediaPlayer->GetDuration() * NewValue);
+		LastUrl = FText::GetEmpty();
+	}
+	else if (Event == EMediaEvent::MediaOpenFailed)
+	{
+		FNotificationInfo NotificationInfo(LOCTEXT("MediaOpenFailedError", "The media failed to open. Check Output Log for details!"));
+		{
+			NotificationInfo.ExpireDuration = 2.0f;
+		}
+
+		FSlateNotificationManager::Get().AddNotification(NotificationInfo)->SetCompletionState(SNotificationItem::CS_Fail);
 	}
 }
 
@@ -823,26 +936,6 @@ void SMediaPlayerEditorViewer::HandleScaleMenuNewMenu(FMenuBuilder& MenuBuilder)
 }
 
 
-FReply SMediaPlayerEditorViewer::HandleTextureMouseButtonUp(const FGeometry& /*Geometry*/, const FPointerEvent& MouseEvent)
-{
-	if (MouseEvent.GetEffectingButton() == EKeys::LeftMouseButton)
-	{
-		if (MediaPlayer->IsPlaying())
-		{
-			MediaPlayer->Pause();
-		}
-		else
-		{
-			MediaPlayer->Play();
-		}
-
-		return FReply::Handled();
-	}
-
-	return FReply::Unhandled();
-}
-
-
 FText SMediaPlayerEditorViewer::HandleTimerTextBlockText() const
 {
 	if (!MediaPlayer->IsReady())
@@ -852,6 +945,7 @@ FText SMediaPlayerEditorViewer::HandleTimerTextBlockText() const
 
 	FTimespan Time = MediaPlayer->GetTime();
 
+	// empty string if time n/a
 	if (Time < FTimespan::Zero())
 	{
 		return FText::GetEmpty();
@@ -859,12 +953,18 @@ FText SMediaPlayerEditorViewer::HandleTimerTextBlockText() const
 
 	FTimespan Duration = MediaPlayer->GetDuration();
 
+	// time only if duration n/a
 	if (Duration <= FTimespan::Zero())
 	{
 		return FText::AsTimespan(Time);
 	}
 
-	return FText::Format(LOCTEXT("TimerTextBlockFormat", "{0} / {1}"), FText::AsTimespan(Time), FText::AsTimespan(Duration));
+	// format time & duration
+	const FText DurationText = (Duration == FTimespan::MaxValue())
+		? FText::FromString(TEXT("\u221E")) // infinity symbol
+		: FText::AsTimespan(Duration);
+
+	return FText::Format(LOCTEXT("TimerTextBlockFormat", "{0} / {1}"), FText::AsTimespan(Time), DurationText);
 }
 
 
@@ -882,6 +982,11 @@ FText SMediaPlayerEditorViewer::HandleTimerTextBlockToolTipText() const
 		return LOCTEXT("UnknownTimeRemainingTooltip", "Unknown time remaining");
 	}
 
+	if (Remaining == FTimespan::MaxValue())
+	{
+		return LOCTEXT("InfiniteTimeRemainingTooltip", "Infinite time remaining");
+	}
+
 	return FText::Format(LOCTEXT("TimeRemainingTooltipFormat", "{0} remaining"), FText::AsTimespan(Remaining));
 }
 
@@ -916,35 +1021,36 @@ void SMediaPlayerEditorViewer::HandleTrackMenuNewMenu(FMenuBuilder& MenuBuilder,
 			const FCulturePtr Culture = I18n.GetCulture(Language);
 			const FString LanguageDisplayName = Culture.IsValid() ? Culture->GetDisplayName() : FString();
 			const FString LanguageNativeName = Culture.IsValid() ? Culture->GetNativeName() : FString();
+			
+			const FText DisplayText = LanguageNativeName.IsEmpty() ? DisplayName : FText::Format(LOCTEXT("TrackNameFormat", "{0} ({1})"), DisplayName, FText::FromString(LanguageNativeName));
+			const FText TooltipText = FText::FromString(LanguageDisplayName);
 
 			MenuBuilder.AddMenuEntry(
-				LanguageNativeName.IsEmpty()
-					? DisplayName
-					: FText::Format(LOCTEXT("TrackNameFormat", "{0} ({1})"), DisplayName, FText::FromString(LanguageNativeName)),
-				FText::FromString(LanguageDisplayName),
+				DisplayText,
+				TooltipText,
 				FSlateIcon(),
 				FUIAction(
-					FExecuteAction::CreateLambda([=]{ MediaPlayer->SelectTrack(TrackType, TrackIndex); }),
+					FExecuteAction::CreateLambda([=] { MediaPlayer->SelectTrack(TrackType, TrackIndex); }),
 					FCanExecuteAction(),
-					FIsActionChecked::CreateLambda([=]{ return (MediaPlayer->GetSelectedTrack(TrackType) == TrackIndex); })
+					FIsActionChecked::CreateLambda([=] { return (MediaPlayer->GetSelectedTrack(TrackType) == TrackIndex); })
 				),
 				NAME_None,
 				EUserInterfaceActionType::RadioButton
 			);
 		}
-
-		return;
 	}
+	else
+	{
+		TSharedRef<SWidget> NoTracksAvailableWidget = SNew(SBox)
+			.Padding(FMargin(2.0f, 2.0f, 2.0f, 0.0f))
+			[
+				SNew(STextBlock)
+					.ColorAndOpacity(FSlateColor::UseSubduedForeground())
+					.Text(LOCTEXT("NoTracksAvailableLabel", "No tracks available"))
+			];
 
-	TSharedRef<SWidget> NoTracksAvailableWidget = SNew(SBox)
-		.Padding(FMargin(2.0f, 2.0f, 2.0f, 0.0f))
-		[
-			SNew(STextBlock)
-				.ColorAndOpacity(FSlateColor::UseSubduedForeground())
-				.Text(LOCTEXT("NoTracksAvailableLabel", "No tracks available"))
-		];
-
-	MenuBuilder.AddWidget(NoTracksAvailableWidget, FText::GetEmpty(), true, false);
+		MenuBuilder.AddWidget(NoTracksAvailableWidget, FText::GetEmpty(), true, false);
+	}
 }
 
 
@@ -952,7 +1058,7 @@ FReply SMediaPlayerEditorViewer::HandleUrlBoxKeyDown(const FGeometry&, const FKe
 {
 	if (KeyEvent.GetKey() == EKeys::Escape)
 	{
-		UrlTextBox->SetText(FText::FromString(MediaPlayer->GetLastUrl()));
+		UrlTextBox->SetText(LastUrl);
 
 		return FReply::Handled().ClearUserFocus(true);
 	}
@@ -961,4 +1067,13 @@ FReply SMediaPlayerEditorViewer::HandleUrlBoxKeyDown(const FGeometry&, const FKe
 }
 
 
+void SMediaPlayerEditorViewer::HandleVideoCaptureDevicesMenuNewMenu(FMenuBuilder& MenuBuilder)
+{
+	TArray<FMediaCaptureDeviceInfo> DeviceInfos;
+	MediaCaptureSupport::EnumerateVideoCaptureDevices(DeviceInfos);
+	MakeCaptureDeviceMenu(DeviceInfos, MenuBuilder);
+}
+
+
+
 #undef LOCTEXT_NAMESPACE
diff --git a/Engine/Plugins/Media/MediaPlayerEditor/Source/MediaPlayerEditor/Private/Widgets/SMediaPlayerEditorViewer.h b/Engine/Plugins/Media/MediaPlayerEditor/Source/MediaPlayerEditor/Private/Widgets/SMediaPlayerEditorViewer.h
--- a/Engine/Plugins/Media/MediaPlayerEditor/Source/MediaPlayerEditor/Private/Widgets/SMediaPlayerEditorViewer.h
+++ b/Engine/Plugins/Media/MediaPlayerEditor/Source/MediaPlayerEditor/Private/Widgets/SMediaPlayerEditorViewer.h
@@ -2,9 +2,12 @@
 
 #pragma once
 
-#include "CoreMinimal.h"
-#include "Input/Reply.h"
+#include "CoreTypes.h"
+#include "Internationalization/Text.h"
+#include "MediaCaptureSupport.h"
 #include "Styling/ISlateStyle.h"
+#include "Templates/SharedPointer.h"
+#include "UObject/NameTypes.h"
 #include "Widgets/DeclarativeSyntaxSupport.h"
 #include "Widgets/SCompoundWidget.h"
 
@@ -12,9 +15,11 @@ class FMenuBuilder;
 class SEditableTextBox;
 class SSlider;
 class UMediaPlayer;
+
 enum class EMediaEvent;
 enum class EMediaPlayerTrack : uint8;
 
+
 /**
  * Implements the contents of the viewer tab in the UMediaPlayer asset editor.
  */
@@ -41,7 +46,7 @@ public:
 	 *
 	 * @param InArgs The declaration data for this widget.
 	 * @param InMediaPlayer The UMediaPlayer asset to show the details for.
-	 * @param InStyleSet The style set to use.
+	 * @param InStyle The style set to use.
 	 */
 	void Construct(const FArguments& InArgs, UMediaPlayer& InMediaPlayer, const TSharedRef<ISlateStyle>& InStyle);
 
@@ -56,6 +61,14 @@ public:
 
 protected:
 
+	/**
+	 * Populate a menu from the given capture device information.
+	 *
+	 * @param DeviceInfos The capture device information.
+	 * @param MenuBuilder The builder for the menu.
+	 */
+	void MakeCaptureDeviceMenu(TArray<FMediaCaptureDeviceInfo>& DeviceInfos, FMenuBuilder& MenuBuilder);
+
 	/** Open the specified media URL. */
 	void OpenUrl(const FText& TextUrl);
 
@@ -64,66 +77,39 @@ protected:
 
 private:
 
-	/** Callback for the active timer. */
-	EActiveTimerReturnType HandleActiveTimer(double InCurrentTime, float InDeltaTime);
-
-	/** Callback for media player events. */
-	void HandleMediaPlayerMediaEvent(EMediaEvent Event);
-
-	/** Callback for getting the text of the player name overlay. */
-	FText HandleOverlayPlayerNameText() const;
-
-	/** Callback for getting the text of the playback state overlay. */
-	FText HandleOverlayStateText() const;
-
-	/** Callback for getting the maximum value of the playback rate spin box. */
-	TOptional<float> HandlePlaybackRateBoxMaxValue() const;
-
-	/** Callback for getting the minimum value of the playback rate spin box. */
-	TOptional<float> HandlePlaybackRateBoxMinValue() const;
-
-	/** Callback for getting the current value of the playback rate spin box. */
-	TOptional<float> HandlePlaybackRateSpinBoxValue() const;
-
-	/** Callback for committing a new value to the playback rate spin box. */
-	void HandlePlaybackRateBoxValueChanged( float NewValue );
+	/** Callback for creating the audio capture devices sub-menu. */
+	void HandleAudioCaptureDevicesMenuNewMenu(FMenuBuilder& MenuBuilder);
 
 	/** Callback for creating the player sub-menu. */
-	void HandlePlayerMenuNewMenu(FMenuBuilder& MenuBuilder);
+	void HandleDecoderMenuNewMenu(FMenuBuilder& MenuBuilder);
 
-	/** Callback for getting the enabled state of the position slider. */
-	bool HandlePositionSliderIsEnabled() const;
+	/** Callback for getting the text of the FPS text block. */
+	FText HandleFpsTextBlockText() const;
 
-	/** Callback for when the position slider captures the mouse. */
-	void HandlePositionSliderMouseCaptureBegin();
+	/** Callback for creating a track format sub-menu. */
+	void HandleFormatMenuNewMenu(FMenuBuilder& MenuBuilder, EMediaPlayerTrack TrackType);
 
-	/** Callback for when the position slider releases the mouse. */
-	void HandlePositionSliderMouseCaptureEnd();
-
-	/** Callback for getting the value of the position slider. */
-	float HandlePositionSliderValue() const;
-
-	/** Callback for changing the value of the 'Position' slider. */
-	void HandlePositionSliderValueChanged( float NewValue );
+	/** Callback for media player events. */
+	void HandleMediaPlayerMediaEvent(EMediaEvent Event);
 
 	/** Callback for creating the Scale sub-menu. */
 	void HandleScaleMenuNewMenu(FMenuBuilder& MenuBuilder);
 
-	/** Callback for handling button up events on the movie texture border. */
-	FReply HandleTextureMouseButtonUp(const FGeometry& Geometry, const FPointerEvent& MouseEvent);
-
 	/** Callback for getting the text of the timer text block. */
 	FText HandleTimerTextBlockText() const;
 
 	/** Callback for getting the tool tip of the timer text block. */
 	FText HandleTimerTextBlockToolTipText() const;
 
-	/** Callback for creating the a track sub-menu. */
+	/** Callback for creating a track sub-menu. */
 	void HandleTrackMenuNewMenu(FMenuBuilder& MenuBuilder, EMediaPlayerTrack TrackType);
 
 	/** Callback for handling key down events in the URL text box. */
 	FReply HandleUrlBoxKeyDown(const FGeometry&, const FKeyEvent& KeyEvent);
 
+	/** Callback for creating the video capture devices sub-menu. */
+	void HandleVideoCaptureDevicesMenuNewMenu(FMenuBuilder& MenuBuilder);
+
 private:
 
 	/** Whether something is currently being dragged over the widget. */
@@ -132,6 +118,9 @@ private:
 	/** Whether the dragged object is a media file that can be played. */
 	bool DragValid;
 
+	/** The text that was last typed into the URL box. */
+	FText LastUrl;
+
 	/** Pointer to the media player that is being viewed. */
 	UMediaPlayer* MediaPlayer;
 
@@ -141,6 +130,9 @@ private:
 	/** Holds the scrubber slider. */
 	TSharedPtr<SSlider> ScrubberSlider;
 
+	/** The value currently being scrubbed to. */
+	float ScrubValue;
+
 	/** The style set to use for this widget. */
 	TSharedPtr<ISlateStyle> Style;
 
diff --git a/Engine/Plugins/Media/MediaPlayerEditor/Source/MediaPlayerEditor/Private/Widgets/SMediaPlayerEditorViewport.cpp b/Engine/Plugins/Media/MediaPlayerEditor/Source/MediaPlayerEditor/Private/Widgets/SMediaPlayerEditorViewport.cpp
new file mode 100644
--- /dev/null
+++ b/Engine/Plugins/Media/MediaPlayerEditor/Source/MediaPlayerEditor/Private/Widgets/SMediaPlayerEditorViewport.cpp
@@ -0,0 +1,321 @@
+// Copyright 1998-2017 Epic Games, Inc. All Rights Reserved.
+
+#include "Widgets/SMediaPlayerEditorViewport.h"
+
+#include "MediaPlayer.h"
+#include "Misc/Paths.h"
+#include "Widgets/SBoxPanel.h"
+#include "Widgets/SOverlay.h"
+#include "Widgets/Layout/SScaleBox.h"
+#include "Widgets/Text/STextBlock.h"
+
+#include "MediaPlayerEditorSettings.h"
+#include "SMediaPlayerEditorOutput.h"
+#include "SMediaPlayerEditorOverlay.h"
+
+
+#define LOCTEXT_NAMESPACE "SMediaPlayerEditorViewport"
+
+
+/* SMediaPlayerEditorPlayer structors
+ *****************************************************************************/
+
+SMediaPlayerEditorViewport::SMediaPlayerEditorViewport()
+	: MediaPlayer(nullptr)
+{ }
+
+
+/* SMediaPlayerEditorPlayer interface
+ *****************************************************************************/
+
+void SMediaPlayerEditorViewport::Construct(const FArguments& InArgs, UMediaPlayer& InMediaPlayer, const TSharedRef<ISlateStyle>& InStyle)
+{
+	MediaPlayer = &InMediaPlayer;
+	Style = InStyle;
+
+	ChildSlot
+	[
+		SNew(SOverlay)
+
+		+ SOverlay::Slot()
+			[
+				// movie viewport
+				SNew(SScaleBox)
+					.Stretch_Lambda([]() -> EStretch::Type {
+						auto Settings = GetDefault<UMediaPlayerEditorSettings>();
+						if (Settings->ViewportScale == EMediaPlayerEditorScale::Fill)
+						{
+							return EStretch::Fill;
+						}
+						if (Settings->ViewportScale == EMediaPlayerEditorScale::Fit)
+						{
+							return EStretch::ScaleToFit;
+						}
+						return EStretch::None;
+					})
+					[
+						// movie texture
+						SNew(SMediaPlayerEditorOutput, InMediaPlayer)
+					]
+			]
+
+		+ SOverlay::Slot()
+			[
+				// subtitle & caption overlays
+				SNew(SMediaPlayerEditorOverlay, InMediaPlayer)
+					.Visibility_Lambda([]{ return GetDefault<UMediaPlayerEditorSettings>()->ShowTextOverlays ? EVisibility::Visible : EVisibility::Collapsed; })
+			]
+
+		+ SOverlay::Slot()
+			.Padding(FMargin(12.0f, 8.0f))
+			[
+				// top info overlays
+				SNew(SHorizontalBox)
+					.Visibility_Lambda([this]() -> EVisibility {
+						return (IsHovered() || HasMouseCapture()) ? EVisibility::Visible : EVisibility::Hidden;
+					})
+
+				+ SHorizontalBox::Slot()
+					.HAlign(HAlign_Left)
+					.VAlign(VAlign_Top)
+					[
+						// media source name
+						SNew(STextBlock)
+							.ColorAndOpacity(FSlateColor::UseSubduedForeground())
+							.Font(FSlateFontInfo(FPaths::EngineContentDir() / TEXT("Slate/Fonts/Roboto-Regular.ttf"), 18))
+							.ShadowOffset(FVector2D(1.f, 1.f))
+							.Text(this, &SMediaPlayerEditorViewport::HandleMediaSourceNameText)
+							.ToolTipText(LOCTEXT("OverlaySourceNameTooltip", "Name of the currently opened media source"))
+					]
+
+				+ SHorizontalBox::Slot()
+					.HAlign(HAlign_Right)
+					.VAlign(VAlign_Top)
+					[
+						// player name
+						SNew(STextBlock)
+							.ColorAndOpacity(FSlateColor::UseSubduedForeground())
+							.Font(FSlateFontInfo(FPaths::EngineContentDir() / TEXT("Slate/Fonts/Roboto-Regular.ttf"), 18))
+							.ShadowOffset(FVector2D(1.f, 1.f))
+							.Text(this, &SMediaPlayerEditorViewport::HandleMediaPlayerNameText)
+							.ToolTipText(LOCTEXT("OverlayPlayerNameTooltip", "Name of the currently used media player plug-in"))
+					]
+			]
+
+		+ SOverlay::Slot()
+			.Padding(FMargin(12.0f, 8.0f))
+			.VAlign(VAlign_Bottom)
+			[
+				// bottom info overlays
+				SNew(SHorizontalBox)
+					.Visibility_Lambda([this]() -> EVisibility {
+						return (IsHovered() || HasMouseCapture()) ? EVisibility::Visible : EVisibility::Hidden;
+					})
+
+				+ SHorizontalBox::Slot()
+					.HAlign(HAlign_Left)
+					.VAlign(VAlign_Bottom)
+					[
+						// playback state
+						SNew(STextBlock)
+							.ColorAndOpacity(FSlateColor::UseSubduedForeground())
+							.Font(FSlateFontInfo(FPaths::EngineContentDir() / TEXT("Slate/Fonts/Roboto-Regular.ttf"), 18))
+							.ShadowOffset(FVector2D(1.f, 1.f))
+							.Text(this, &SMediaPlayerEditorViewport::HandleMediaPlayerStateText)
+							.ToolTipText(LOCTEXT("OverlayPlayerStateTooltip", "The media player's current state"))
+					]
+
+				+ SHorizontalBox::Slot()
+					.HAlign(HAlign_Right)
+					.VAlign(VAlign_Bottom)
+					[
+						// view settings
+						SNew(STextBlock)
+							.ColorAndOpacity(FSlateColor::UseSubduedForeground())
+							.Font(FSlateFontInfo(FPaths::EngineContentDir() / TEXT("Slate/Fonts/Roboto-Regular.ttf"), 18))
+							.ShadowOffset(FVector2D(1.f, 1.f))
+							.Text(this, &SMediaPlayerEditorViewport::HandleViewSettingsText)
+							.ToolTipText(LOCTEXT("OverlayViewSettingsTooltip", "The current view settings"))
+					]
+			]
+	];
+}
+
+
+/* SWidget interface
+ *****************************************************************************/
+
+FReply SMediaPlayerEditorViewport::OnMouseButtonDown(const FGeometry& MyGeometry, const FPointerEvent& MouseEvent)
+{
+	if (MouseEvent.GetEffectingButton() != EKeys::LeftMouseButton)
+	{
+		return FReply::Unhandled();
+	}
+
+	SetCursor(EMouseCursor::CardinalCross);
+
+	return FReply::Handled().CaptureMouse(SharedThis(this));
+}
+
+
+FReply SMediaPlayerEditorViewport::OnMouseButtonDoubleClick(const FGeometry& MyGeometry, const FPointerEvent& MouseEvent)
+{
+	if (MouseEvent.GetEffectingButton() != EKeys::LeftMouseButton)
+	{
+		return FReply::Unhandled();
+	}
+
+	if (MediaPlayer->IsPlaying())
+	{
+		MediaPlayer->Pause();
+	}
+	else
+	{
+		MediaPlayer->Play();
+	}
+
+	return FReply::Handled();
+}
+
+
+FReply SMediaPlayerEditorViewport::OnMouseButtonUp(const FGeometry& MyGeometry, const FPointerEvent& MouseEvent)
+{
+	if (MouseEvent.GetEffectingButton() != EKeys::LeftMouseButton)
+	{
+		return FReply::Unhandled();
+	}
+
+	SetCursor(EMouseCursor::Default);
+
+	return FReply::Handled().ReleaseMouseCapture();
+}
+
+
+FReply SMediaPlayerEditorViewport::OnMouseMove(const FGeometry& MyGeometry, const FPointerEvent& MouseEvent)
+{
+	if (!HasMouseCapture())
+	{
+		return FReply::Unhandled();
+	}
+
+	const FVector2D CursorDelta = MouseEvent.GetCursorDelta();
+	MediaPlayer->SetViewRotation(FRotator::MakeFromEuler(FVector(CursorDelta.Y, CursorDelta.X, 0.0f)), false);
+
+	return FReply::Handled();
+}
+
+
+FReply SMediaPlayerEditorViewport::OnMouseWheel(const FGeometry& MyGeometry, const FPointerEvent& MouseEvent)
+{
+	MediaPlayer->SetViewField(MouseEvent.GetWheelDelta(), 0.0f, false);
+
+	return FReply::Handled();
+}
+
+
+/* SMediaPlayerEditorViewport callbacks
+ *****************************************************************************/
+
+FText SMediaPlayerEditorViewport::HandleMediaPlayerNameText() const
+{
+	FName PlayerName = MediaPlayer->GetPlayerName();
+
+	if ((PlayerName == NAME_None) || MediaPlayer->GetUrl().IsEmpty())
+	{
+		const FName DesiredPlayerName = MediaPlayer->GetDesiredPlayerName();
+
+		if (DesiredPlayerName == NAME_None)
+		{
+			return LOCTEXT("AutoPlayerName", "Auto");
+		}
+
+		return FText::FromName(DesiredPlayerName);
+	}
+
+	return FText::FromName(PlayerName);
+}
+
+
+FText SMediaPlayerEditorViewport::HandleMediaPlayerStateText() const
+{
+	if (MediaPlayer->HasError())
+	{
+		return LOCTEXT("StateOverlayError", "Error");
+	}
+
+	if (MediaPlayer->IsBuffering())
+	{
+		return LOCTEXT("StateOverlayBuffering", "Buffering");
+	}
+
+	if (MediaPlayer->IsPreparing())
+	{
+		return LOCTEXT("StateOverlayPreparing", "Preparing");
+	}
+
+	if (!MediaPlayer->IsReady())
+	{
+		return LOCTEXT("StateOverlayStopped", "Not Ready");
+	}
+
+	if (MediaPlayer->IsPaused())
+	{
+		return LOCTEXT("StateOverlayPaused", "Paused");
+	}
+
+	if (MediaPlayer->IsPlaying())
+	{
+		float Rate = MediaPlayer->GetRate();
+
+		if (Rate == 1.0f)
+		{
+			return LOCTEXT("StateOverlayPlaying", "Playing");
+		}
+
+		if (Rate < 0.0f)
+		{
+			return FText::Format(LOCTEXT("StateOverlayReverseFormat", "Reverse ({0}x)"), FText::AsNumber(-Rate));
+		}
+
+		return FText::Format(LOCTEXT("StateOverlayForwardFormat", "Forward ({0}x)"), FText::AsNumber(Rate));
+	}
+
+	return LOCTEXT("StateOverlayReady", "Ready");
+}
+
+
+FText SMediaPlayerEditorViewport::HandleMediaSourceNameText() const
+{
+	if (MediaPlayer->GetUrl().IsEmpty())
+	{
+		return LOCTEXT("StateOverlayNoMedia", "No Media");
+	}
+
+	return FText::FromString(FPaths::GetBaseFilename(MediaPlayer->GetUrl()));
+}
+
+
+FText SMediaPlayerEditorViewport::HandleViewSettingsText() const
+{
+	FVector Euler = MediaPlayer->GetViewRotation().Euler();
+
+	FNumberFormattingOptions NumberFormat;
+	{
+		NumberFormat.MaximumFractionalDigits = 0;
+		NumberFormat.MinimumIntegralDigits = 3;
+	}
+
+	FFormatNamedArguments Args;
+	{
+		Args.Add(TEXT("P"), FText::AsNumber(Euler.Y, &NumberFormat));
+		Args.Add(TEXT("Y"), FText::AsNumber(Euler.Z, &NumberFormat));
+		Args.Add(TEXT("R"), FText::AsNumber(Euler.X, &NumberFormat));
+		Args.Add(TEXT("H"), FText::AsNumber(MediaPlayer->GetHorizontalFieldOfView(), &NumberFormat));
+		Args.Add(TEXT("V"), FText::AsNumber(MediaPlayer->GetVerticalFieldOfView(), &NumberFormat));
+	}
+
+	return FText::Format(LOCTEXT("ViewSettingsFormat", "R {P} {Y} {R} | FOV {H} {V}"), Args);
+}
+
+
+#undef LOCTEXT_NAMESPACE
diff --git a/Engine/Plugins/Media/MediaPlayerEditor/Source/MediaPlayerEditor/Private/Widgets/SMediaPlayerEditorViewport.h b/Engine/Plugins/Media/MediaPlayerEditor/Source/MediaPlayerEditor/Private/Widgets/SMediaPlayerEditorViewport.h
new file mode 100644
--- /dev/null
+++ b/Engine/Plugins/Media/MediaPlayerEditor/Source/MediaPlayerEditor/Private/Widgets/SMediaPlayerEditorViewport.h
@@ -0,0 +1,67 @@
+// Copyright 1998-2017 Epic Games, Inc. All Rights Reserved.
+
+#pragma once
+
+#include "Styling/ISlateStyle.h"
+#include "Widgets/DeclarativeSyntaxSupport.h"
+#include "Widgets/SCompoundWidget.h"
+
+class UMediaPlayer;
+
+
+class SMediaPlayerEditorViewport
+	: public SCompoundWidget
+{
+public:
+
+	SLATE_BEGIN_ARGS(SMediaPlayerEditorViewport) { }
+	SLATE_END_ARGS()
+
+public:
+
+	/** Default constructor. */
+	SMediaPlayerEditorViewport();
+
+public:
+
+	/**
+	 * Construct this widget
+	 *
+	 * @param InArgs The declaration data for this widget.
+	 * @param InMediaPlayer The UMediaPlayer asset to show the details for.
+	 * @param InStyleSet The style set to use.
+	 */
+	void Construct(const FArguments& InArgs, UMediaPlayer& InMediaPlayer, const TSharedRef<ISlateStyle>& InStyle);
+
+public:
+
+	//~ SWidget interface
+
+	virtual FReply OnMouseButtonDown(const FGeometry& MyGeometry, const FPointerEvent& MouseEvent) override;
+	virtual FReply OnMouseButtonDoubleClick(const FGeometry& MyGeometry, const FPointerEvent& MouseEvent) override;
+	virtual FReply OnMouseButtonUp(const FGeometry& MyGeometry, const FPointerEvent& MouseEvent) override;
+	virtual FReply OnMouseMove(const FGeometry& MyGeometry, const FPointerEvent& MouseEvent) override;
+	virtual FReply OnMouseWheel(const FGeometry& MyGeometry, const FPointerEvent& MouseEvent) override;
+
+private:
+
+	/** Callback for getting the text of the player name overlay. */
+	FText HandleMediaPlayerNameText() const;
+
+	/** Callback for getting the text of the playback state overlay. */
+	FText HandleMediaPlayerStateText() const;
+
+	/** Callback for getting the text of the media source name overlay. */
+	FText HandleMediaSourceNameText() const;
+
+	/** Callback for getting the text of the view settings overlay. */
+	FText HandleViewSettingsText() const;
+
+private:
+
+	/** Pointer to the media player that is being viewed. */
+	UMediaPlayer* MediaPlayer;
+
+	/** The style set to use for this widget. */
+	TSharedPtr<ISlateStyle> Style;
+};
\ No newline at end of file
diff --git a/Engine/Plugins/Media/MfMedia/MfMedia.uplugin b/Engine/Plugins/Media/MfMedia/MfMedia.uplugin
--- a/Engine/Plugins/Media/MfMedia/MfMedia.uplugin
+++ b/Engine/Plugins/Media/MfMedia/MfMedia.uplugin
@@ -10,7 +10,7 @@
 	"DocsURL" : "https://forums.unrealengine.com/showthread.php?46879-Media-Framework-Documentation-for-4-5-Preview",
 	"MarketplaceURL" : "",
 	"SupportURL" : "",
-	"EnabledByDefault" : false,
+	"EnabledByDefault" : true,
 	"CanContainContent" : false,
 	"IsBetaVersion" : false,
 	"Installed" : false,
@@ -20,7 +20,7 @@
 			"Name" : "MfMedia",
 			"Type" : "RuntimeNoCommandlet",
 			"LoadingPhase" : "PreLoadingScreen",
-			"WhitelistPlatforms" : [ "XboxOne" ]
+			"WhitelistPlatforms" : [ "XboxOne", "Win32", "Win64" ]
 		},
 		{
 			"Name" : "MfMediaEditor",
@@ -36,7 +36,7 @@
 			"Name" : "MfMediaFactory",
 			"Type" : "RuntimeNoCommandlet",
 			"LoadingPhase" : "PostEngineInit",
-			"WhitelistPlatforms" : [ "XboxOne" ]
+			"WhitelistPlatforms" : [ "XboxOne", "Win32", "Win64" ]
 		}
 	]
 }
diff --git a/Engine/Plugins/Media/MfMedia/Source/MfMedia/MfMedia.Build.cs b/Engine/Plugins/Media/MfMedia/Source/MfMedia/MfMedia.Build.cs
--- a/Engine/Plugins/Media/MfMedia/Source/MfMedia/MfMedia.Build.cs
+++ b/Engine/Plugins/Media/MfMedia/Source/MfMedia/MfMedia.Build.cs
@@ -12,35 +12,33 @@ namespace UnrealBuildTool.Rules
 			DynamicallyLoadedModuleNames.AddRange(
 				new string[] {
 					"Media",
-				}
-			);
+				});
 
 			PrivateDependencyModuleNames.AddRange(
 				new string[] {
 					"Core",
+					"MediaUtils",
 					"RenderCore",
-				}
-			);
+				});
 
 			PrivateIncludePathModuleNames.AddRange(
 				new string[] {
 					"Media",
-				}
-			);
+				});
 
 			PrivateIncludePaths.AddRange(
 				new string[] {
 					"MfMedia/Private",
 					"MfMedia/Private/Mf",
 					"MfMedia/Private/Player",
-				}
-			);
+				});
 
 			PublicAdditionalLibraries.Add("mfplat.lib");
 			PublicAdditionalLibraries.Add("mfreadwrite.lib");
 			PublicAdditionalLibraries.Add("mfuuid.lib");
 
-			if ((Target.Platform == UnrealTargetPlatform.Win32) || (Target.Platform == UnrealTargetPlatform.Win64))
+			if ((Target.Platform == UnrealTargetPlatform.Win32) ||
+				(Target.Platform == UnrealTargetPlatform.Win64))
 			{
 				PublicAdditionalLibraries.Add("mf.lib");
 				PublicAdditionalLibraries.Add("Propsys.lib");
diff --git a/Engine/Plugins/Media/MfMedia/Source/MfMedia/Private/Mf/MfMediaByteStream.cpp b/Engine/Plugins/Media/MfMedia/Source/MfMedia/Private/Mf/MfMediaByteStream.cpp
--- a/Engine/Plugins/Media/MfMedia/Source/MfMedia/Private/Mf/MfMediaByteStream.cpp
+++ b/Engine/Plugins/Media/MfMedia/Source/MfMedia/Private/Mf/MfMediaByteStream.cpp
@@ -1,11 +1,15 @@
 // Copyright 1998-2017 Epic Games, Inc. All Rights Reserved.
 
 #include "MfMediaByteStream.h"
-#include "Misc/ScopeLock.h"
 
 #if MFMEDIA_SUPPORTED_PLATFORM
 
+#include "CoreTypes.h"
+#include "HAL/PlatformAtomics.h"
+#include "Misc/ScopeLock.h"
 #include "MfMediaReadState.h"
+#include "Serialization/Archive.h"
+#include "Windows/COMPointer.h"
 
 #if PLATFORM_WINDOWS
 	#include "AllowWindowsPlatformTypes.h"
diff --git a/Engine/Plugins/Media/MfMedia/Source/MfMedia/Private/Mf/MfMediaByteStream.h b/Engine/Plugins/Media/MfMedia/Source/MfMedia/Private/Mf/MfMediaByteStream.h
--- a/Engine/Plugins/Media/MfMedia/Source/MfMedia/Private/Mf/MfMediaByteStream.h
+++ b/Engine/Plugins/Media/MfMedia/Source/MfMedia/Private/Mf/MfMediaByteStream.h
@@ -2,17 +2,22 @@
 
 #pragma once
 
-#include "../MfMediaPrivate.h"
+#include "MfMediaPrivate.h"
 
 #if MFMEDIA_SUPPORTED_PLATFORM
 
+#include "CoreTypes.h"
+#include "Templates/SharedPointer.h"
+
 #if PLATFORM_WINDOWS
-	#include "WindowsHWrapper.h"
-	#include "AllowWindowsPlatformTypes.h"
+	#include "Windows/WindowsHWrapper.h"
+	#include "Windows/AllowWindowsPlatformTypes.h"
 #else
-	#include "XboxOneAllowPlatformTypes.h"
+	#include "XboxOne/XboxOneAllowPlatformTypes.h"
 #endif
 
+class FArchive;
+
 
 /**
  * Implements a wrapper for media source data that is played from memory.
@@ -30,9 +35,6 @@ public:
 	 */
 	FMfMediaByteStream(const TSharedRef<FArchive, ESPMode::ThreadSafe>& InArchive);
 
-	/** Virtual destructor. */
-	virtual ~FMfMediaByteStream() { }
-
 public:
 
 	//~ IMFAsyncCallback interface
@@ -63,6 +65,11 @@ public:
 	STDMETHODIMP SetCurrentPosition(QWORD qwPosition);
 	STDMETHODIMP Write(const BYTE* pb, ULONG cb, ULONG* pcbWritten);
 
+protected:
+
+	/** Hidden destructor. */
+	virtual ~FMfMediaByteStream() { }
+
 protected:
 
 	/**
@@ -101,9 +108,9 @@ private:
 
 
 #if PLATFORM_WINDOWS
-	#include "HideWindowsPlatformTypes.h"
+	#include "Windows/HideWindowsPlatformTypes.h"
 #else
-	#include "XboxOneHidePlatformTypes.h"
+	#include "XboxOne/XboxOneHidePlatformTypes.h"
 #endif
 
-#endif
+#endif //MFMEDIA_SUPPORTED_PLATFORM
diff --git a/Engine/Plugins/Media/MfMedia/Source/MfMedia/Private/Mf/MfMediaReadState.h b/Engine/Plugins/Media/MfMedia/Source/MfMedia/Private/Mf/MfMediaReadState.h
--- a/Engine/Plugins/Media/MfMedia/Source/MfMedia/Private/Mf/MfMediaReadState.h
+++ b/Engine/Plugins/Media/MfMedia/Source/MfMedia/Private/Mf/MfMediaReadState.h
@@ -2,13 +2,15 @@
 
 #pragma once
 
-#include "../MfMediaPrivate.h"
+#include "MfMediaPrivate.h"
+
+#if MFMEDIA_SUPPORTED_PLATFORM
 
 #if PLATFORM_WINDOWS
-	#include "WindowsHWrapper.h"
-	#include "AllowWindowsPlatformTypes.h"
+	#include "Windows/WindowsHWrapper.h"
+	#include "Windows/AllowWindowsPlatformTypes.h"
 #else
-	#include "XboxOneAllowPlatformTypes.h"
+	#include "XboxOne/XboxOneAllowPlatformTypes.h"
 #endif
 
 
@@ -37,9 +39,6 @@ public:
 		, RefCount(0)
 	{ }
 
-	/** Virtual destructor. */
-	virtual ~FMfMediaReadState() { }
-
 public:
 
 	/**
@@ -123,6 +122,11 @@ public:
 		return CurrentRefCount;
 	}
 
+protected:
+
+	/** Hidden destructor. */
+	virtual ~FMfMediaReadState() { }
+
 private:
 
 	/** Number of bytes read. */
@@ -140,7 +144,9 @@ private:
 
 
 #if PLATFORM_WINDOWS
-	#include "HideWindowsPlatformTypes.h"
+	#include "Windows/HideWindowsPlatformTypes.h"
 #else
-	#include "XboxOneHidePlatformTypes.h"
+	#include "XboxOne/XboxOneHidePlatformTypes.h"
 #endif
+
+#endif //MFMEDIA_SUPPORTED_PLATFORM
diff --git a/Engine/Plugins/Media/MfMedia/Source/MfMedia/Private/Mf/MfMediaUtils.cpp b/Engine/Plugins/Media/MfMedia/Source/MfMedia/Private/Mf/MfMediaUtils.cpp
--- a/Engine/Plugins/Media/MfMedia/Source/MfMedia/Private/Mf/MfMediaUtils.cpp
+++ b/Engine/Plugins/Media/MfMedia/Source/MfMedia/Private/Mf/MfMediaUtils.cpp
@@ -4,6 +4,15 @@
 
 #if MFMEDIA_SUPPORTED_PLATFORM
 
+#include "HAL/FileManager.h"
+#include "HAL/PlatformProcess.h"
+#include "Misc/FileHelper.h"
+#include "Serialization/Archive.h"
+#include "Serialization/ArrayReader.h"
+
+#include "MfMediaByteStream.h"
+#include "MfMediaPrivate.h"
+
 #if PLATFORM_WINDOWS
 	#include "AllowWindowsPlatformTypes.h"
 #else
@@ -13,6 +22,171 @@
 
 namespace MfMedia
 {
+	TComPtr<IMFMediaType> CreateOutputType(const GUID& MajorType, const GUID& SubType, bool AllowNonStandardCodecs)
+	{
+		TComPtr<IMFMediaType> OutputType;
+		{
+			HRESULT Result = ::MFCreateMediaType(&OutputType);
+
+			if (FAILED(Result))
+			{
+				UE_LOG(LogMfMedia, Warning, TEXT("Failed to create %s output type (%s)"), *MajorTypeToString(MajorType), *ResultToString(Result));
+				return NULL;
+			}
+
+			Result = OutputType->SetUINT32(MF_MT_ALL_SAMPLES_INDEPENDENT, TRUE);
+
+			if (FAILED(Result))
+			{
+				UE_LOG(LogMfMedia, Warning, TEXT("Failed to initialize %s output type (%s)"), *MajorTypeToString(MajorType), *ResultToString(Result));
+				return NULL;
+			}
+		}
+
+		if (MajorType == MFMediaType_Audio)
+		{
+#if PLATFORM_XBOXONE
+			// filter unsupported audio formats (XboxOne only supports AAC)
+			if (SubType != MFAudioFormat_AAC)
+			{
+				UE_LOG(LogMfMedia, Warning, TEXT("Skipping unsupported audio type %s (%s) \"%s\""), *SubTypeToString(SubType), *GuidToString(SubType), *FourccToString(SubType.Data1));
+				return NULL;
+			}
+#else
+			// filter unsupported audio formats
+			if (FMemory::Memcmp(&SubType.Data2, &MFMPEG4Format_Base.Data2, 12) == 0)
+			{
+				if (AllowNonStandardCodecs)
+				{
+					UE_LOG(LogMfMedia, Verbose, TEXT("Allowing non-standard MP4 audio type %s (%s) \"%s\""), *SubTypeToString(SubType), *GuidToString(SubType), *FourccToString(SubType.Data1));
+				}
+				else
+				{
+					const bool DocumentedFormat =
+						(SubType.Data1 == WAVE_FORMAT_ADPCM) ||
+						(SubType.Data1 == WAVE_FORMAT_ALAW) ||
+						(SubType.Data1 == WAVE_FORMAT_MULAW) ||
+						(SubType.Data1 == WAVE_FORMAT_IMA_ADPCM) ||
+						(SubType.Data1 == MFAudioFormat_AAC.Data1) ||
+						(SubType.Data1 == MFAudioFormat_MP3.Data1) ||
+						(SubType.Data1 == MFAudioFormat_PCM.Data1);
+
+					const bool UndocumentedFormat =
+						(SubType.Data1 == WAVE_FORMAT_WMAUDIO2) ||
+						(SubType.Data1 == WAVE_FORMAT_WMAUDIO3) ||
+						(SubType.Data1 == WAVE_FORMAT_WMAUDIO_LOSSLESS);
+
+					if (!DocumentedFormat && !UndocumentedFormat)
+					{
+						UE_LOG(LogMfMedia, Warning, TEXT("Skipping non-standard MP4 audio type %s (%s) \"%s\""), *SubTypeToString(SubType), *GuidToString(SubType), *FourccToString(SubType.Data1));
+						return NULL;
+					}
+				}
+			}
+			else if (FMemory::Memcmp(&SubType.Data2, &MFAudioFormat_Base.Data2, 12) != 0)
+			{
+				if (AllowNonStandardCodecs)
+				{
+					UE_LOG(LogMfMedia, Verbose, TEXT("Allowing non-standard audio type %s (%s) \"%s\""), *SubTypeToString(SubType), *GuidToString(SubType), *FourccToString(SubType.Data1));
+				}
+				else
+				{
+					UE_LOG(LogMfMedia, Warning, TEXT("Skipping non-standard audio type %s (%s) \"%s\""), *SubTypeToString(SubType), *GuidToString(SubType), *FourccToString(SubType.Data1));
+					return NULL;
+				}
+			}
+#endif //PLATFORM_XBOXONE
+
+			// configure audio output
+			if (FAILED(OutputType->SetGUID(MF_MT_MAJOR_TYPE, MFMediaType_Audio)) ||
+				FAILED(OutputType->SetGUID(MF_MT_SUBTYPE, MFAudioFormat_PCM)) ||
+				FAILED(OutputType->SetUINT32(MF_MT_AUDIO_BITS_PER_SAMPLE, 16u)))
+				//FAILED(OutputType->SetGUID(MF_MT_SUBTYPE, MFAudioFormat_Float)) ||
+				//FAILED(OutputType->SetUINT32(MF_MT_AUDIO_BITS_PER_SAMPLE, 32u)))
+			{
+				UE_LOG(LogMfMedia, Warning, TEXT("Failed to initialize audio output type"));
+				return NULL;
+			}
+		}
+		else if (MajorType == MFMediaType_SAMI)
+		{
+			// configure caption output
+			const HRESULT Result = OutputType->SetGUID(MF_MT_MAJOR_TYPE, MFMediaType_SAMI);
+
+			if (FAILED(Result))
+			{
+				UE_LOG(LogMfMedia, Warning, TEXT("Failed to initialize caption output type (%s)"), *ResultToString(Result));
+				return NULL;
+			}
+		}
+		else if (MajorType == MFMediaType_Video)
+		{
+#if PLATFORM_XBOXONE
+			// filter unsupported video types (XboxOne only supports H.264)
+			if ((SubType != MFVideoFormat_H264) && (SubType == MFVideoFormat_H264_ES))
+			{
+				UE_LOG(LogMfMedia, Warning, TEXT("Unsupported video type %s (%s) \"%i\""), *SubTypeToString(SubType), *GuidToString(SubType), *FourccToString(SubType.Data1));
+				return NULL;
+			}
+#else
+			// filter unsupported video types
+			if (FMemory::Memcmp(&SubType.Data2, &MFVideoFormat_Base.Data2, 12) != 0)
+			{
+				if (AllowNonStandardCodecs)
+				{
+					UE_LOG(LogMfMedia, Verbose, TEXT("Allowing non-standard video type %s (%s) \"%s\""), *SubTypeToString(SubType), *GuidToString(SubType), *FourccToString(SubType.Data1));
+				}
+				else
+				{
+					UE_LOG(LogMfMedia, Warning, TEXT("Skipping non-standard video type %s (%s) \"%s\""), *SubTypeToString(SubType), *GuidToString(SubType), *FourccToString(SubType.Data1));
+					return NULL;
+				}
+			}
+
+			if ((SubType == MFVideoFormat_HEVC) && !FWindowsPlatformMisc::VerifyWindowsVersion(10, 0))
+			{
+				UE_LOG(LogMfMedia, Warning, TEXT("Your Windows version is %s"), *FPlatformMisc::GetOSVersion());
+
+				if ((SubType == MFVideoFormat_HEVC) && !FWindowsPlatformMisc::VerifyWindowsVersion(6, 2))
+				{
+					UE_LOG(LogMfMedia, Warning, TEXT("HEVC video type requires Windows 10 or newer"));
+					return NULL;
+				}
+
+				UE_LOG(LogMfMedia, Warning, TEXT("HEVC video type requires Windows 10 or newer (game must be manifested for Windows 10)"));
+			}
+#endif //PLATFORM_XBOXONE
+
+			// configure video output
+			HRESULT Result = OutputType->SetGUID(MF_MT_MAJOR_TYPE, MFMediaType_Video);
+
+			if (FAILED(Result))
+			{
+				UE_LOG(LogMfMedia, Warning, TEXT("Failed to set video output type (%s)"), *ResultToString(Result));
+				return NULL;
+			}
+
+#if 1//PLATFORM_XBOXONE
+			Result = OutputType->SetGUID(MF_MT_SUBTYPE, MFVideoFormat_NV12); // XboxOne only supports NV12
+#else
+			Result = OutputType->SetGUID(MF_MT_SUBTYPE, (SubType == MFVideoFormat_RGB32) ? MFVideoFormat_RGB32 : MFVideoFormat_NV12);
+#endif
+
+			if (FAILED(Result))
+			{
+				UE_LOG(LogMfMedia, Warning, TEXT("Failed to set video output sub-type (%s)"), *ResultToString(Result));
+				return NULL;
+			}
+		}
+		else
+		{
+			return NULL; // unsupported input type
+		}
+
+		return OutputType;
+	}
+
+
 	FString FourccToString(unsigned long Fourcc)
 	{
 		return FString::Printf(TEXT("%c%c%c%c"),
@@ -54,6 +228,218 @@
 	}
 
 
+	FString MediaEventToString(MediaEventType Event)
+	{
+		switch (Event)
+		{
+		case MEUnknown: return TEXT("Unknown");
+		case MEError: return TEXT("Error");
+		case MEExtendedType: return TEXT("Extended Type");
+		case MENonFatalError: return TEXT("Non-fatal Error");
+		case MESessionUnknown: return TEXT("Session Unknown");
+		case MESessionTopologySet: return TEXT("Session Topology Set");
+		case MESessionTopologiesCleared: return TEXT("Session Topologies Cleared");
+		case MESessionStarted: return TEXT("Session Started");
+		case MESessionPaused: return TEXT("Session Paused");
+		case MESessionStopped: return TEXT("Session Stopped");
+		case MESessionClosed: return TEXT("Session Closed");
+		case MESessionEnded: return TEXT("Session Ended");
+		case MESessionRateChanged: return TEXT("Session Rate Changed");
+		case MESessionScrubSampleComplete: return TEXT("Session Scrub Sample Complete");
+		case MESessionCapabilitiesChanged: return TEXT("Session Capabilities Changed");
+		case MESessionTopologyStatus: return TEXT("Session Topology Status");
+		case MESessionNotifyPresentationTime: return TEXT("Session Notify Presentation Time");
+		case MENewPresentation: return TEXT("New Presentation");
+		case MELicenseAcquisitionStart: return TEXT("License Acquisition Start");
+		case MELicenseAcquisitionCompleted: return TEXT("License Acquisition Completed");
+		case MEIndividualizationStart: return TEXT("Individualization Start");
+		case MEIndividualizationCompleted: return TEXT("Individualization Completed");
+		case MEEnablerProgress: return TEXT("Enabler Progress");
+		case MEEnablerCompleted: return TEXT("Enabler Completed");
+		case MEPolicyError: return TEXT("Policy Error");
+		case MEPolicyReport: return TEXT("Policy Report");
+		case MEBufferingStarted: return TEXT("Buffering Started");
+		case MEBufferingStopped: return TEXT("Buffering Stopped");
+		case MEConnectStart: return TEXT("Connect Start");
+		case MEConnectEnd: return TEXT("Connect End");
+		case MEReconnectStart: return TEXT("Reconnect Start");
+		case MEReconnectEnd: return TEXT("Reconnect End");
+		case MERendererEvent: return TEXT("Renderer Event");
+		case MESessionStreamSinkFormatChanged: return TEXT("Session Stream Sink Format Changed");
+		case MESourceUnknown: return TEXT("Source Unknown");
+		case MESourceStarted: return TEXT("Source Started");
+		case MEStreamStarted: return TEXT("Stream Started");
+		case MESourceSeeked: return TEXT("Source Seeked");
+		case MEStreamSeeked: return TEXT("Stream Seeked");
+		case MENewStream: return TEXT("New Stream");
+		case MEUpdatedStream: return TEXT("Updated Stream");
+		case MESourceStopped: return TEXT("Source Stopped");
+		case MEStreamStopped: return TEXT("Stream Stopped");
+		case MESourcePaused: return TEXT("Source Paused");
+		case MEStreamPaused: return TEXT("Stream Paused");
+		case MEEndOfPresentation: return TEXT("End of Presentation");
+		case MEEndOfStream: return TEXT("End of Stream");
+		case MEMediaSample: return TEXT("Media Sample");
+		case MEStreamTick: return TEXT("Stream Tick");
+		case MEStreamThinMode: return TEXT("Stream Thin Mode");
+		case MEStreamFormatChanged: return TEXT("Stream Format Changed");
+		case MESourceRateChanged: return TEXT("Source Rate Changed");
+		case MEEndOfPresentationSegment: return TEXT("End of Presentation Segment");
+		case MESourceCharacteristicsChanged: return TEXT("Source Characteristics Changed");
+		case MESourceRateChangeRequested: return TEXT("Source Rate Change Requested");
+		case MESourceMetadataChanged: return TEXT("Source Metadata Changed");
+		case MESequencerSourceTopologyUpdated: return TEXT("Sequencer Source Topology Updated");
+		case MESinkUnknown: return TEXT("Sink Unknown");
+		case MEStreamSinkStarted: return TEXT("Stream Sink Started");
+		case MEStreamSinkStopped: return TEXT("Stream Sink Stopped");
+		case MEStreamSinkPaused: return TEXT("Strema Sink Paused");
+		case MEStreamSinkRateChanged: return TEXT("Stream Sink Rate Changed");
+		case MEStreamSinkRequestSample: return TEXT("Stream Sink Request Sample");
+		case MEStreamSinkMarker: return TEXT("Stream Sink Marker");
+		case MEStreamSinkPrerolled: return TEXT("Stream Sink Prerolled");
+		case MEStreamSinkScrubSampleComplete: return TEXT("Stream Sink Scrub Sample Complete");
+		case MEStreamSinkFormatChanged: return TEXT("Stream Sink Format Changed");
+		case MEStreamSinkDeviceChanged: return TEXT("Stream Sink Device Changed");
+		case MEQualityNotify: return TEXT("Quality Notify");
+		case MESinkInvalidated: return TEXT("Sink Invalidated");
+		case MEAudioSessionNameChanged: return TEXT("Audio Session Name Changed");
+		case MEAudioSessionVolumeChanged: return TEXT("Audio Session Volume Changed");
+		case MEAudioSessionDeviceRemoved: return TEXT("Audio Session Device Removed");
+		case MEAudioSessionServerShutdown: return TEXT("Audio Session Server Shutdown");
+		case MEAudioSessionGroupingParamChanged: return TEXT("Audio Session Grouping Param Changed");
+		case MEAudioSessionIconChanged: return TEXT("Audio Session Icion Changed");
+		case MEAudioSessionFormatChanged: return TEXT("Audio Session Format Changed");
+		case MEAudioSessionDisconnected: return TEXT("Audio Session Disconnected");
+		case MEAudioSessionExclusiveModeOverride: return TEXT("Audio Session Exclusive Mode Override");
+		case MECaptureAudioSessionVolumeChanged: return TEXT("Capture Audio Session Volume Changed");
+		case MECaptureAudioSessionDeviceRemoved: return TEXT("Capture Audio Session Device Removed");
+		case MECaptureAudioSessionFormatChanged: return TEXT("Capture Audio Session Format Changed");
+		case MECaptureAudioSessionDisconnected: return TEXT("Capture Audio Session Disconnected");
+		case MECaptureAudioSessionExclusiveModeOverride: return TEXT("Capture Audio Session Exclusive Mode Override");
+		case MECaptureAudioSessionServerShutdown: return TEXT("Capture Audio Session Server Shutdown");
+		case METrustUnknown: return TEXT("Trust Unknown");
+		case MEPolicyChanged: return TEXT("Policy Changed");
+		case MEContentProtectionMessage: return TEXT("Content Protection Message");
+		case MEPolicySet: return TEXT("Policy Set");
+		case MEWMDRMLicenseBackupCompleted: return TEXT("WM DRM License Backup Completed");
+		case MEWMDRMLicenseBackupProgress: return TEXT("WM DRM License Backup Progress");
+		case MEWMDRMLicenseRestoreCompleted: return TEXT("WM DRM License Restore Completed");
+		case MEWMDRMLicenseRestoreProgress: return TEXT("WM DRM License Restore Progress");
+		case MEWMDRMLicenseAcquisitionCompleted: return TEXT("WM DRM License Acquisition Completed");
+		case MEWMDRMIndividualizationCompleted: return TEXT("WM DRM Individualization Completed");
+		case MEWMDRMIndividualizationProgress: return TEXT("WM DRM Individualization Progress");
+		case MEWMDRMProximityCompleted: return TEXT("WM DRM Proximity Completed");
+		case MEWMDRMLicenseStoreCleaned: return TEXT("WM DRM License Store Cleaned");
+		case MEWMDRMRevocationDownloadCompleted: return TEXT("WM DRM Revocation Download Completed");
+		case METransformUnknown: return TEXT("Transform Unkonwn");
+		case METransformNeedInput: return TEXT("Transform Need Input");
+		case METransformHaveOutput: return TEXT("Transform Have Output");
+		case METransformDrainComplete: return TEXT("Transform Drain Complete");
+		case METransformMarker: return TEXT("Transform Marker");
+		case MEByteStreamCharacteristicsChanged: return TEXT("Byte Stream Characteristics Changed");
+		case MEVideoCaptureDeviceRemoved: return TEXT("Video Capture Device Removed");
+		case MEVideoCaptureDevicePreempted: return TEXT("Video Capture Device Preempted");
+
+		default:
+			return FString::Printf(TEXT("Unknown event %i"), Event);
+		}
+	}
+
+
+	TComPtr<IMFMediaSource> ResolveMediaSource(TSharedPtr<FArchive, ESPMode::ThreadSafe> Archive, const FString& Url, bool Precache)
+	{
+		// load media source
+		if (!Archive.IsValid() && Url.StartsWith(TEXT("file://")))
+		{
+			const TCHAR* FilePath = &Url[7];
+
+			if (Precache)
+			{
+				FArrayReader* Reader = new FArrayReader;
+
+				if (FFileHelper::LoadFileToArray(*Reader, FilePath))
+				{
+					Archive = MakeShareable(Reader);
+				}
+				else
+				{
+					delete Reader;
+				}
+			}
+			else
+			{
+				Archive = MakeShareable(IFileManager::Get().CreateFileReader(FilePath));
+			}
+
+			if (!Archive.IsValid())
+			{
+				UE_LOG(LogMfMedia, Error, TEXT("Failed to open or read media file %s"), FilePath);
+				return NULL;
+			}
+
+			if (Archive->TotalSize() == 0)
+			{
+				UE_LOG(LogMfMedia, Error, TEXT("Cannot open media from empty file %s."), FilePath);
+				return NULL;
+			}
+		}
+
+		// create source resolver
+		TComPtr<IMFSourceResolver> SourceResolver;
+		{
+			const HRESULT Result = ::MFCreateSourceResolver(&SourceResolver);
+
+			if (FAILED(Result))
+			{
+				UE_LOG(LogMfMedia, Error, TEXT("Failed to create media source resolver (%s)"), *MfMedia::ResultToString(Result));
+				return NULL;
+			}
+		}
+
+		// resolve media source
+		TComPtr<IUnknown> SourceObject;
+		{
+			MF_OBJECT_TYPE ObjectType;
+
+			if (Archive.IsValid())
+			{
+				TComPtr<FMfMediaByteStream> ByteStream = new FMfMediaByteStream(Archive.ToSharedRef());
+				const HRESULT Result = SourceResolver->CreateObjectFromByteStream(ByteStream, *Url, MF_RESOLUTION_MEDIASOURCE, NULL, &ObjectType, &SourceObject);
+
+				if (FAILED(Result))
+				{
+					UE_LOG(LogMfMedia, Error, TEXT("Failed to resolve byte stream %s (%s)"), *Url, *MfMedia::ResultToString(Result));
+					return NULL;
+				}
+			}
+			else
+			{
+				const HRESULT Result = SourceResolver->CreateObjectFromURL(*Url, MF_RESOLUTION_MEDIASOURCE, NULL, &ObjectType, &SourceObject);
+
+				if (FAILED(Result))
+				{
+					UE_LOG(LogMfMedia, Error, TEXT("Failed to resolve URL %s (%s)"), *Url, *MfMedia::ResultToString(Result));
+					return NULL;
+				}
+			}
+		}
+
+		// get media source interface
+		TComPtr<IMFMediaSource> MediaSource;
+		{
+			const HRESULT Result = SourceObject->QueryInterface(IID_PPV_ARGS(&MediaSource));
+
+			if (FAILED(Result))
+			{
+				UE_LOG(LogMfMedia, Error, TEXT("Failed to query media source interface: %s"), *MfMedia::ResultToString(Result));
+				return NULL;
+			}
+		}
+
+		return MediaSource;
+	}
+
+
 	FString ResultToString(HRESULT Result)
 	{
 		void* DllHandle = nullptr;
@@ -131,6 +517,77 @@
 
 	FString SubTypeToString(const GUID& SubType)
 	{
+		if (SubType == GUID_NULL) return TEXT("Null");
+
+		// image formats
+		if (SubType == MFImageFormat_JPEG) return TEXT("Jpeg");
+		if (SubType == MFImageFormat_RGB32) return TEXT("RGB32");
+
+		// stream formats
+		if (SubType == MFStreamFormat_MPEG2Transport) return TEXT("MPEG-2 Transport");
+		if (SubType == MFStreamFormat_MPEG2Program) return TEXT("MPEG-2 Program");
+
+		// video formats
+		if (SubType == MFVideoFormat_RGB32) return TEXT("RGB32");
+		if (SubType == MFVideoFormat_ARGB32) return TEXT("ARGB32");
+		if (SubType == MFVideoFormat_RGB24) return TEXT("RGB24");
+		if (SubType == MFVideoFormat_RGB555) return TEXT("RGB525");
+		if (SubType == MFVideoFormat_RGB565) return TEXT("RGB565");
+		if (SubType == MFVideoFormat_RGB8) return TEXT("RGB8");
+		if (SubType == MFVideoFormat_AI44) return TEXT("AI44");
+		if (SubType == MFVideoFormat_AYUV) return TEXT("AYUV");
+		if (SubType == MFVideoFormat_YUY2) return TEXT("YUY2");
+		if (SubType == MFVideoFormat_YVYU) return TEXT("YVYU");
+		if (SubType == MFVideoFormat_YVU9) return TEXT("YVU9");
+		if (SubType == MFVideoFormat_UYVY) return TEXT("UYVY");
+		if (SubType == MFVideoFormat_NV11) return TEXT("NV11");
+		if (SubType == MFVideoFormat_NV12) return TEXT("NV12");
+		if (SubType == MFVideoFormat_YV12) return TEXT("YV12");
+		if (SubType == MFVideoFormat_I420) return TEXT("I420");
+		if (SubType == MFVideoFormat_IYUV) return TEXT("IYUV");
+		if (SubType == MFVideoFormat_Y210) return TEXT("Y210");
+		if (SubType == MFVideoFormat_Y216) return TEXT("Y216");
+		if (SubType == MFVideoFormat_Y410) return TEXT("Y410");
+		if (SubType == MFVideoFormat_Y416) return TEXT("Y416");
+		if (SubType == MFVideoFormat_Y41P) return TEXT("Y41P");
+		if (SubType == MFVideoFormat_Y41T) return TEXT("Y41T");
+		if (SubType == MFVideoFormat_Y42T) return TEXT("Y42T");
+		if (SubType == MFVideoFormat_P210) return TEXT("P210");
+		if (SubType == MFVideoFormat_P216) return TEXT("P216");
+		if (SubType == MFVideoFormat_P010) return TEXT("P010");
+		if (SubType == MFVideoFormat_P016) return TEXT("P016");
+		if (SubType == MFVideoFormat_v210) return TEXT("v210");
+		if (SubType == MFVideoFormat_v216) return TEXT("v216");
+		if (SubType == MFVideoFormat_v410) return TEXT("v410");
+		if (SubType == MFVideoFormat_MP43) return TEXT("MP43");
+		if (SubType == MFVideoFormat_MP4S) return TEXT("MP4S");
+		if (SubType == MFVideoFormat_M4S2) return TEXT("M4S2");
+		if (SubType == MFVideoFormat_MP4V) return TEXT("MP4V");
+		if (SubType == MFVideoFormat_WMV1) return TEXT("WMV1");
+		if (SubType == MFVideoFormat_WMV2) return TEXT("WMV2");
+		if (SubType == MFVideoFormat_WMV3) return TEXT("WMV3");
+		if (SubType == MFVideoFormat_WVC1) return TEXT("WVC1");
+		if (SubType == MFVideoFormat_MSS1) return TEXT("MSS1");
+		if (SubType == MFVideoFormat_MSS2) return TEXT("MSS2");
+		if (SubType == MFVideoFormat_MPG1) return TEXT("MPG1");
+		if (SubType == MFVideoFormat_DVSL) return TEXT("DVSL");
+		if (SubType == MFVideoFormat_DVSD) return TEXT("DVSD");
+		if (SubType == MFVideoFormat_DVHD) return TEXT("DVHD");
+		if (SubType == MFVideoFormat_DV25) return TEXT("DV25");
+		if (SubType == MFVideoFormat_DV50) return TEXT("DV50");
+		if (SubType == MFVideoFormat_DVH1) return TEXT("DVH1");
+		if (SubType == MFVideoFormat_DVC) return TEXT("DVC");
+		if (SubType == MFVideoFormat_H264) return TEXT("H264");
+		if (SubType == MFVideoFormat_MJPG) return TEXT("MJPG");
+		if (SubType == MFVideoFormat_420O) return TEXT("420O");
+
+#if (WINVER >= _WIN32_WINNT_WIN8)
+		if (SubType == MFVideoFormat_H263) return TEXT("H263");
+#endif
+
+		if (SubType == MFVideoFormat_H264_ES) return TEXT("H264 ES");
+		if (SubType == MFVideoFormat_MPEG2) return TEXT("MPEG-2");
+
 		// audio formats
 		if ((FMemory::Memcmp(&SubType.Data2, &MFAudioFormat_Base.Data2, 12) == 0) ||
 			(FMemory::Memcmp(&SubType.Data2, &MFMPEG4Format_Base.Data2, 12) == 0))
@@ -398,75 +855,7 @@
 			if (SubType.Data1 == WAVE_FORMAT_FLAC) return TEXT("flac.sourceforge.net");
 		}
 
-		// image formats
-		if (SubType == MFImageFormat_JPEG) return TEXT("Jpeg");
-		if (SubType == MFImageFormat_RGB32) return TEXT("RGB32");
-
-		// stream formats
-		if (SubType == MFStreamFormat_MPEG2Transport) return TEXT("MPEG-2 Transport");
-		if (SubType == MFStreamFormat_MPEG2Program) return TEXT("MPEG-2 Program");
-
-		// video formats
-		if (SubType == MFVideoFormat_RGB32) return TEXT("RGB32");
-		if (SubType == MFVideoFormat_ARGB32) return TEXT("ARGB32");
-		if (SubType == MFVideoFormat_RGB24) return TEXT("RGB24");
-		if (SubType == MFVideoFormat_RGB555) return TEXT("RGB525");
-		if (SubType == MFVideoFormat_RGB565) return TEXT("RGB565");
-		if (SubType == MFVideoFormat_RGB8) return TEXT("RGB8");
-		if (SubType == MFVideoFormat_AI44) return TEXT("AI44");
-		if (SubType == MFVideoFormat_AYUV) return TEXT("AYUV");
-		if (SubType == MFVideoFormat_YUY2) return TEXT("YUY2");
-		if (SubType == MFVideoFormat_YVYU) return TEXT("YVYU");
-		if (SubType == MFVideoFormat_YVU9) return TEXT("YVU9");
-		if (SubType == MFVideoFormat_UYVY) return TEXT("UYVY");
-		if (SubType == MFVideoFormat_NV11) return TEXT("NV11");
-		if (SubType == MFVideoFormat_NV12) return TEXT("NV12");
-		if (SubType == MFVideoFormat_YV12) return TEXT("YV12");
-		if (SubType == MFVideoFormat_I420) return TEXT("I420");
-		if (SubType == MFVideoFormat_IYUV) return TEXT("IYUV");
-		if (SubType == MFVideoFormat_Y210) return TEXT("Y210");
-		if (SubType == MFVideoFormat_Y216) return TEXT("Y216");
-		if (SubType == MFVideoFormat_Y410) return TEXT("Y410");
-		if (SubType == MFVideoFormat_Y416) return TEXT("Y416");
-		if (SubType == MFVideoFormat_Y41P) return TEXT("Y41P");
-		if (SubType == MFVideoFormat_Y41T) return TEXT("Y41T");
-		if (SubType == MFVideoFormat_Y42T) return TEXT("Y42T");
-		if (SubType == MFVideoFormat_P210) return TEXT("P210");
-		if (SubType == MFVideoFormat_P216) return TEXT("P216");
-		if (SubType == MFVideoFormat_P010) return TEXT("P010");
-		if (SubType == MFVideoFormat_P016) return TEXT("P016");
-		if (SubType == MFVideoFormat_v210) return TEXT("v210");
-		if (SubType == MFVideoFormat_v216) return TEXT("v216");
-		if (SubType == MFVideoFormat_v410) return TEXT("v410");
-		if (SubType == MFVideoFormat_MP43) return TEXT("MP43");
-		if (SubType == MFVideoFormat_MP4S) return TEXT("MP4S");
-		if (SubType == MFVideoFormat_M4S2) return TEXT("M4S2");
-		if (SubType == MFVideoFormat_MP4V) return TEXT("MP4V");
-		if (SubType == MFVideoFormat_WMV1) return TEXT("WMV1");
-		if (SubType == MFVideoFormat_WMV2) return TEXT("WMV2");
-		if (SubType == MFVideoFormat_WMV3) return TEXT("WMV3");
-		if (SubType == MFVideoFormat_WVC1) return TEXT("WVC1");
-		if (SubType == MFVideoFormat_MSS1) return TEXT("MSS1");
-		if (SubType == MFVideoFormat_MSS2) return TEXT("MSS2");
-		if (SubType == MFVideoFormat_MPG1) return TEXT("MPG1");
-		if (SubType == MFVideoFormat_DVSL) return TEXT("DVSL");
-		if (SubType == MFVideoFormat_DVSD) return TEXT("DVSD");
-		if (SubType == MFVideoFormat_DVHD) return TEXT("DVHD");
-		if (SubType == MFVideoFormat_DV25) return TEXT("DV25");
-		if (SubType == MFVideoFormat_DV50) return TEXT("DV50");
-		if (SubType == MFVideoFormat_DVH1) return TEXT("DVH1");
-		if (SubType == MFVideoFormat_DVC) return TEXT("DVC");
-		if (SubType == MFVideoFormat_H264) return TEXT("H264");
-		if (SubType == MFVideoFormat_MJPG) return TEXT("MJPG");
-		if (SubType == MFVideoFormat_420O) return TEXT("420O");
-
-#if (WINVER >= _WIN32_WINNT_WIN8)
-		if (SubType == MFVideoFormat_H263) return TEXT("H263");
-#endif
-
-		if (SubType == MFVideoFormat_H264_ES) return TEXT("H264 ES");
-		if (SubType == MFVideoFormat_MPEG2) return TEXT("MPEG-2");
-
+		// unknown type
 		return FString::Printf(TEXT("%s (%s)"), *GuidToString(SubType), *FourccToString(SubType.Data1));
 	}
 }
diff --git a/Engine/Plugins/Media/MfMedia/Source/MfMedia/Private/Mf/MfMediaUtils.h b/Engine/Plugins/Media/MfMedia/Source/MfMedia/Private/Mf/MfMediaUtils.h
--- a/Engine/Plugins/Media/MfMedia/Source/MfMedia/Private/Mf/MfMediaUtils.h
+++ b/Engine/Plugins/Media/MfMedia/Source/MfMedia/Private/Mf/MfMediaUtils.h
@@ -2,20 +2,35 @@
 
 #pragma once
 
-#include "../MfMediaPrivate.h"
+#include "MfMediaPrivate.h"
 
 #if MFMEDIA_SUPPORTED_PLATFORM
 
+#include "Containers/UnrealString.h"
+#include "Templates/SharedPointer.h"
+
 #if PLATFORM_WINDOWS
-	#include "WindowsHWrapper.h"
-	#include "AllowWindowsPlatformTypes.h"
+	#include "Windows/WindowsHWrapper.h"
+	#include "Windows/AllowWindowsPlatformTypes.h"
 #else
-	#include "XboxOneAllowPlatformTypes.h"
+	#include "XboxOne/XboxOneAllowPlatformTypes.h"
 #endif
 
+class FArchive;
+
 
 namespace MfMedia
 {
+	/**
+	 * Create an output media type for the given input media type.
+	 *
+	 * @param MajorType The output's major type, i.e. audio or video.
+	 * @param SubType The output's sub type, i.e. PCM or RGB32.
+	 * @param AllowNonStandardCodecs Whether non-standard codecs should be allowed.
+	 * @return The output type, or NULL on error.
+	 */
+	TComPtr<IMFMediaType> CreateOutputType(const GUID& MajorType, const GUID& SubType, bool AllowNonStandardCodecs);
+
 	/** 
 	 * Convert a FOURCC code to string.
 	 *
@@ -40,6 +55,24 @@ namespace MfMedia
 	 */
 	FString MajorTypeToString(const GUID& MajorType);
 
+	/**
+	 * Convert a media event to string.
+	 *
+	 * @param Event The event code to convert.
+	 * @return The corresponding string.
+	 */
+	FString MediaEventToString(MediaEventType Event);
+
+	/**
+	 * Resolve a media source from an archive or URL.
+	 *
+	 * @param Archive The archive to read media data from (optional).
+	 * @param Url The media source URL.
+	 * @param Precache Whether to precache media into RAM if URL is a local file.
+	 * @return The media source object, or NULL if it couldn't be resolved.
+	 */
+	TComPtr<IMFMediaSource> ResolveMediaSource(TSharedPtr<FArchive, ESPMode::ThreadSafe> Archive, const FString& Url, bool Precache);
+
 	/**
 	 * Convert an WMF HRESULT code to string.
 	 *
@@ -59,9 +92,9 @@ namespace MfMedia
 
 
 #if PLATFORM_WINDOWS
-	#include "HideWindowsPlatformTypes.h"
+	#include "Windows/HideWindowsPlatformTypes.h"
 #else
-	#include "XboxOneHidePlatformTypes.h"
+	#include "XboxOne/XboxOneHidePlatformTypes.h"
 #endif
 
-#endif
+#endif //MFMEDIA_SUPPORTED_PLATFORM
diff --git a/Engine/Plugins/Media/MfMedia/Source/MfMedia/Private/MfMediaModule.cpp b/Engine/Plugins/Media/MfMedia/Source/MfMedia/Private/MfMediaModule.cpp
--- a/Engine/Plugins/Media/MfMedia/Source/MfMedia/Private/MfMediaModule.cpp
+++ b/Engine/Plugins/Media/MfMedia/Source/MfMedia/Private/MfMediaModule.cpp
@@ -1,16 +1,18 @@
 // Copyright 1998-2017 Epic Games, Inc. All Rights Reserved.
 
 #include "MfMediaPrivate.h"
-#include "IMfMediaModule.h"
+
+#include "Modules/ModuleManager.h"
 
 #if MFMEDIA_SUPPORTED_PLATFORM
 	#include "MfMediaPlayer.h"
+	#include "Templates/SharedPointer.h"
 #endif
 
+#include "IMfMediaModule.h"
 
-DEFINE_LOG_CATEGORY(LogMfMedia);
 
-#define LOCTEXT_NAMESPACE "FMfMediaModule"
+DEFINE_LOG_CATEGORY(LogMfMedia);
 
 
 /**
@@ -30,12 +32,12 @@ public:
 
 	//~ IMfMediaModule interface
 
-	virtual TSharedPtr<IMediaPlayer, ESPMode::ThreadSafe> CreatePlayer() override
+	virtual TSharedPtr<IMediaPlayer, ESPMode::ThreadSafe> CreatePlayer(IMediaEventSink& EventSink) override
 	{
 #if MFMEDIA_SUPPORTED_PLATFORM
 		if (Initialized)
 		{
-			return MakeShareable(new FMfMediaPlayer());
+			return MakeShared<FMfMediaPlayer, ESPMode::ThreadSafe>(EventSink);
 		}
 #endif
 		return nullptr;
@@ -83,6 +85,3 @@ private:
 
 
 IMPLEMENT_MODULE(FMfMediaModule, MfMedia);
-
-
-#undef LOCTEXT_NAMESPACE
diff --git a/Engine/Plugins/Media/MfMedia/Source/MfMedia/Private/MfMediaPrivate.h b/Engine/Plugins/Media/MfMedia/Source/MfMedia/Private/MfMediaPrivate.h
--- a/Engine/Plugins/Media/MfMedia/Source/MfMedia/Private/MfMediaPrivate.h
+++ b/Engine/Plugins/Media/MfMedia/Source/MfMedia/Private/MfMediaPrivate.h
@@ -2,50 +2,39 @@
 
 #pragma once
 
-// you can enable this plug-in on Windows by compiling the Engine and Editor
-// against Windows 7. By default, UE4 is targeting Windows Vista. The following
-// changes have to be made in SetupEnvironment() in UEBuildWindows.cs:
-//
-//   // Windows Vista or higher required
-//   InBuildTarget.GlobalCompileEnvironment.Config.Definitions.Add("_WIN32_WINNT=0x0601");
-//   InBuildTarget.GlobalCompileEnvironment.Config.Definitions.Add("WINVER=0x0601");
+#include "Logging/LogMacros.h"
 
-//#define MFMEDIA_SUPPORTED_PLATFORM (PLATFORM_XBOXONE || (PLATFORM_WINDOWS && WINVER >= 0x0601))
-#define MFMEDIA_SUPPORTED_PLATFORM PLATFORM_XBOXONE // Windows disabled until 4.16, because broken
-
-
-#include "Runtime/Core/Public/CoreMinimal.h"
-#include "Runtime/Core/Public/Modules/ModuleManager.h"
-#include "Runtime/Core/Public/Misc/ScopeLock.h"
+#define MFMEDIA_SUPPORTED_PLATFORM (PLATFORM_XBOXONE || (PLATFORM_WINDOWS && WINVER >= 0x0601 /*Win7*/))
 
 #if MFMEDIA_SUPPORTED_PLATFORM
-	#include "Runtime/Core/Public/Async/Async.h"
-	#include "Runtime/Media/Public/IMediaAudioSink.h"
-	#include "Runtime/Media/Public/IMediaModule.h"
-	#include "Runtime/Media/Public/IMediaOptions.h"
-	#include "Runtime/Media/Public/IMediaOverlaySink.h"
-	#include "Runtime/Media/Public/IMediaPlayerFactory.h"
-	#include "Runtime/Media/Public/IMediaTextureSink.h"
-	#include "Runtime/RenderCore/Public/RenderingThread.h"
-
 	#if PLATFORM_WINDOWS
 		#include "WindowsHWrapper.h"
 	#endif
-	#include "Runtime/Core/Public/Windows/AllowWindowsPlatformTypes.h"
+
+	#include "Windows/AllowWindowsPlatformTypes.h"
+
 	#if PLATFORM_WINDOWS
 		#include <windows.h>
 		#include <propvarutil.h>
 		#include <shlwapi.h>
+
+		const GUID FORMAT_VideoInfo = { 0x05589f80, 0xc356, 0x11ce, { 0xbf, 0x01, 0x00, 0xaa, 0x00, 0x55, 0x59, 0x5a } };
+		const GUID FORMAT_VideoInfo2 = { 0xf72a76A0, 0xeb0a, 0x11d0, { 0xac, 0xe4, 0x00, 0x00, 0xc0, 0xcc, 0x16, 0xba } };
+
+		#if (WINVER < _WIN32_WINNT_WIN8)
+			const GUID MF_LOW_LATENCY = { 0x9c27891a, 0xed7a, 0x40e1, { 0x88, 0xe8, 0xb2, 0x27, 0x27, 0xa0, 0x24, 0xee } };
+		#endif
 	#endif
+
 	#include <mfapi.h>
 	#include <mferror.h>
 	#include <mfidl.h>
 	#include <Mfreadwrite.h>
-	#include "Runtime/Core/Public/Windows/COMPointer.h"
-	#include "Runtime/Core/Public/Windows/HideWindowsPlatformTypes.h"
+
+	#include "Windows/COMPointer.h"
+	#include "Windows/HideWindowsPlatformTypes.h"
 
 #elif PLATFORM_WINDOWS && !UE_SERVER
-	#include "CoreDefines.h"
 	#pragma message("Skipping MfMedia (requires WINVER >= 0x0601, but WINVER is " PREPROCESSOR_TO_STRING(WINVER) ")")
 
 #endif //MFMEDIA_SUPPORTED_PLATFORM
