Fix compile errors with latest version of MSVC with /permissive- enabled.

#rnx

diff --git a/Engine/Plugins/Media/MediaPlayerEditor/Source/MediaPlayerEditor/Public/Widgets/SMediaPermutationsSelector.h b/Engine/Plugins/Media/MediaPlayerEditor/Source/MediaPlayerEditor/Public/Widgets/SMediaPermutationsSelector.h
--- a/Engine/Plugins/Media/MediaPlayerEditor/Source/MediaPlayerEditor/Public/Widgets/SMediaPermutationsSelector.h
+++ b/Engine/Plugins/Media/MediaPlayerEditor/Source/MediaPlayerEditor/Public/Widgets/SMediaPermutationsSelector.h
@@ -122,7 +122,7 @@ public:
 		: _ColumnHeight(200)
 	{}
 		/** Array of columns */
-		SLATE_SUPPORTS_SLOT_WITH_ARGS(typename FColumn)
+		SLATE_SUPPORTS_SLOT_WITH_ARGS(FColumn)
 		/** Array of data items that we are displaying */
 		SLATE_ARGUMENT(TArray<ItemType>, PermutationsSource)
 		/** Default selected item in ItemsSource */
diff --git a/Engine/Plugins/Media/MediaPlayerEditor/Source/MediaPlayerEditor/Public/Widgets/SMediaPermutationsSelector.inl b/Engine/Plugins/Media/MediaPlayerEditor/Source/MediaPlayerEditor/Public/Widgets/SMediaPermutationsSelector.inl
--- a/Engine/Plugins/Media/MediaPlayerEditor/Source/MediaPlayerEditor/Public/Widgets/SMediaPermutationsSelector.inl
+++ b/Engine/Plugins/Media/MediaPlayerEditor/Source/MediaPlayerEditor/Public/Widgets/SMediaPermutationsSelector.inl
@@ -8,7 +8,7 @@
 template <typename ItemType, typename ItemBuilder>
 typename SMediaPermutationsSelector<ItemType, ItemBuilder>::FColumn::FArguments SMediaPermutationsSelector<ItemType, ItemBuilder>::Column(const FName& InColumnName)
 {
-	FColumn::FArguments NewArgs;
+	typename FColumn::FArguments NewArgs;
 	NewArgs._ColumnName = InColumnName;
 	return NewArgs;
 }
diff --git a/Engine/Source/Runtime/Core/Public/GenericPlatform/OSAllocationPool.h b/Engine/Source/Runtime/Core/Public/GenericPlatform/OSAllocationPool.h
--- a/Engine/Source/Runtime/Core/Public/GenericPlatform/OSAllocationPool.h
+++ b/Engine/Source/Runtime/Core/Public/GenericPlatform/OSAllocationPool.h
@@ -307,6 +307,6 @@ public:
 
 	void PrintDebugInfo()
 	{
-		printf("BlockSize: %llu NumAllocated/TotalBlocks = %lu/%lu\n", (uint64)BlockSize, (NumBlocks - NumFreeBlocks), NumBlocks);
+		printf("BlockSize: %llu NumAllocated/TotalBlocks = %llu/%llu\n", (uint64)BlockSize, (uint64)(NumBlocks - NumFreeBlocks), (uint64)NumBlocks);
 	}
 };
diff --git a/Engine/Source/Runtime/SlateCore/Public/Widgets/DeclarativeSyntaxSupport.h b/Engine/Source/Runtime/SlateCore/Public/Widgets/DeclarativeSyntaxSupport.h
--- a/Engine/Source/Runtime/SlateCore/Public/Widgets/DeclarativeSyntaxSupport.h
+++ b/Engine/Source/Runtime/SlateCore/Public/Widgets/DeclarativeSyntaxSupport.h
@@ -382,7 +382,7 @@ template<typename WidgetType> struct TSlateBaseNamedArgs;
  */
 #define SLATE_SUPPORTS_SLOT_WITH_ARGS( SlotType ) \
 	TArray< SlotType* > Slots; \
-	WidgetArgsType& operator + (const SlotType::FArguments& ArgumentsForNewSlot) \
+	WidgetArgsType& operator + (const typename SlotType::FArguments& ArgumentsForNewSlot) \
 		{ \
 			Slots.Add( new SlotType( ArgumentsForNewSlot ) ); \
 			return *this; \
