Add hover style for media player scrub thumb

diff --git a/Engine/Plugins/Media/MediaPlayerEditor/Source/MediaPlayerEditor/Private/Shared/MediaPlayerEditorStyle.cpp b/Engine/Plugins/Media/MediaPlayerEditor/Source/MediaPlayerEditor/Private/Shared/MediaPlayerEditorStyle.cpp
--- a/Engine/Plugins/Media/MediaPlayerEditor/Source/MediaPlayerEditor/Private/Shared/MediaPlayerEditorStyle.cpp
+++ b/Engine/Plugins/Media/MediaPlayerEditor/Source/MediaPlayerEditor/Private/Shared/MediaPlayerEditorStyle.cpp
@@ -75,6 +75,7 @@ FMediaPlayerEditorStyle::FMediaPlayerEditorStyle()
 		.SetNormalBarImage(FSlateColorBrush(FColor::White))
 		.SetDisabledBarImage(FSlateColorBrush(FLinearColor::Gray))
 		.SetNormalThumbImage(IMAGE_BRUSH("scrubber", FVector2D(2.0f, 10.0f)))
+		.SetHoveredThumbImage(IMAGE_BRUSH("scrubber", FVector2D(2.0f, 10.0f)))
 		.SetDisabledThumbImage(IMAGE_BRUSH("scrubber", FVector2D(2.0f, 10.0f)))
 		.SetBarThickness(2.0f)
 	);
