File media source thumbnails optimized, removed overlay icon that displayed precached state

diff --git a/Engine/Plugins/Media/MediaPlayerEditor/Source/MediaPlayerEditor/Private/AssetTools/FileMediaSourceActions.cpp b/Engine/Plugins/Media/MediaPlayerEditor/Source/MediaPlayerEditor/Private/AssetTools/FileMediaSourceActions.cpp
--- a/Engine/Plugins/Media/MediaPlayerEditor/Source/MediaPlayerEditor/Private/AssetTools/FileMediaSourceActions.cpp
+++ b/Engine/Plugins/Media/MediaPlayerEditor/Source/MediaPlayerEditor/Private/AssetTools/FileMediaSourceActions.cpp
@@ -39,31 +39,4 @@ UClass* FFileMediaSourceActions::GetSupportedClass() const
 	return UFileMediaSource::StaticClass();
 }
 
-
-TSharedPtr<class SWidget> FFileMediaSourceActions::GetThumbnailOverlay(const FAssetData& AssetData) const
-{
-	TWeakObjectPtr<UFileMediaSource> FileMediaSource = Cast<UFileMediaSource>(AssetData.GetAsset());
-
-	if (!FileMediaSource.IsValid())
-	{
-		return nullptr;
-	}
-
-	return SNew(SBorder)
-		.BorderImage(FEditorStyle::GetNoBrush())
-		.Visibility_Lambda([=]() -> EVisibility {
-			return (FileMediaSource.IsValid() && FileMediaSource->PrecacheFile)
-				? EVisibility::HitTestInvisible
-				: EVisibility::Hidden;
-		})
-		.Padding(FMargin(0.0f, 0.0f, 2.0f, 7.0f))
-		.HAlign(HAlign_Right)
-		.VAlign(VAlign_Bottom)
-		[
-			SNew(SImage)
-				.Image(Style->GetBrush("MediaPlayerEditor.FileMediaSourcePrecached"))
-		];
-}
-
-
 #undef LOCTEXT_NAMESPACE
diff --git a/Engine/Plugins/Media/MediaPlayerEditor/Source/MediaPlayerEditor/Private/AssetTools/FileMediaSourceActions.h b/Engine/Plugins/Media/MediaPlayerEditor/Source/MediaPlayerEditor/Private/AssetTools/FileMediaSourceActions.h
--- a/Engine/Plugins/Media/MediaPlayerEditor/Source/MediaPlayerEditor/Private/AssetTools/FileMediaSourceActions.h
+++ b/Engine/Plugins/Media/MediaPlayerEditor/Source/MediaPlayerEditor/Private/AssetTools/FileMediaSourceActions.h
@@ -31,7 +31,6 @@ public:
 	virtual bool CanFilter() override;
 	virtual FText GetName() const override;
 	virtual UClass* GetSupportedClass() const override;
-	virtual TSharedPtr<class SWidget> GetThumbnailOverlay(const FAssetData& AssetData) const override;
 
 private:
 
diff --git a/Engine/Plugins/Media/MediaPlayerEditor/Source/MediaPlayerEditor/Private/Shared/MediaPlayerEditorStyle.cpp b/Engine/Plugins/Media/MediaPlayerEditor/Source/MediaPlayerEditor/Private/Shared/MediaPlayerEditorStyle.cpp
--- a/Engine/Plugins/Media/MediaPlayerEditor/Source/MediaPlayerEditor/Private/Shared/MediaPlayerEditorStyle.cpp
+++ b/Engine/Plugins/Media/MediaPlayerEditor/Source/MediaPlayerEditor/Private/Shared/MediaPlayerEditorStyle.cpp
@@ -41,7 +41,6 @@ FMediaPlayerEditorStyle::FMediaPlayerEditorStyle()
 
 	// misc
 	Set("MediaPlayerEditor.DragDropBorder", new BOX_BRUSH("border_dragdrop", 0.5f));
-	Set("MediaPlayerEditor.FileMediaSourcePrecached", new IMAGE_BRUSH("filemediasource_precached", Icon12x12));
 	Set("MediaPlayerEditor.MediaSourceOpened", new IMAGE_BRUSH("mediasource_opened", Icon8x8));
 
 	// tabs
