Converted menu: Media Player Editor
#rnx

diff --git a/Engine/Plugins/Media/MediaPlayerEditor/Source/MediaPlayerEditor/MediaPlayerEditor.Build.cs b/Engine/Plugins/Media/MediaPlayerEditor/Source/MediaPlayerEditor/MediaPlayerEditor.Build.cs
--- a/Engine/Plugins/Media/MediaPlayerEditor/Source/MediaPlayerEditor/MediaPlayerEditor.Build.cs
+++ b/Engine/Plugins/Media/MediaPlayerEditor/Source/MediaPlayerEditor/MediaPlayerEditor.Build.cs
@@ -53,6 +53,7 @@ public class MediaPlayerEditor : ModuleRules
 				"RHI",
 				"TextureEditor",
 				"UnrealEd",
+				"EditorMenus",
 			});
 
 		PrivateIncludePathModuleNames.AddRange(
diff --git a/Engine/Plugins/Media/MediaPlayerEditor/Source/MediaPlayerEditor/Private/MediaPlayerEditorModule.cpp b/Engine/Plugins/Media/MediaPlayerEditor/Source/MediaPlayerEditor/Private/MediaPlayerEditorModule.cpp
--- a/Engine/Plugins/Media/MediaPlayerEditor/Source/MediaPlayerEditor/Private/MediaPlayerEditorModule.cpp
+++ b/Engine/Plugins/Media/MediaPlayerEditor/Source/MediaPlayerEditor/Private/MediaPlayerEditorModule.cpp
@@ -39,6 +39,8 @@
 #include "Models/MediaPlayerEditorCommands.h"
 #include "Shared/MediaPlayerEditorStyle.h"
 #include "Visualizers/MediaSoundComponentVisualizer.h"
+#include "EditorMenuSubsystem.h"
+#include "Widgets/SMediaPlayerEditorMedia.h"
 
 #include "MediaPlayerEditorLog.h"
 
@@ -90,6 +92,8 @@ public:
 
 		FMediaPlayerEditorCommands::Register();
 
+		SMediaPlayerEditorMedia::RegisterMenus();
+
 		RegisterAssetTools();
 		RegisterCustomizations();
 		RegisterEditorDelegates();
diff --git a/Engine/Plugins/Media/MediaPlayerEditor/Source/MediaPlayerEditor/Private/Shared/MediaPlayerEditorMediaContext.h b/Engine/Plugins/Media/MediaPlayerEditor/Source/MediaPlayerEditor/Private/Shared/MediaPlayerEditorMediaContext.h
new file mode 100644
--- /dev/null
+++ b/Engine/Plugins/Media/MediaPlayerEditor/Source/MediaPlayerEditor/Private/Shared/MediaPlayerEditorMediaContext.h
@@ -0,0 +1,34 @@
+// Copyright 1998-2019 Epic Games, Inc. All Rights Reserved.
+
+#pragma once
+
+#include "CoreMinimal.h"
+#include "UObject/ObjectMacros.h"
+#include "UObject/Object.h"
+#include "EditorMenuSubsystem.h"
+
+#include "MediaPlayerEditorMediaContext.generated.h"
+
+class SMediaPlayerEditorMedia;
+
+UCLASS(BlueprintType)
+class MEDIAPLAYEREDITOR_API UMediaPlayerEditorMediaContext : public UEditorMenuContextBase
+{
+	GENERATED_BODY()
+
+public:	
+
+	void InitContext(UObject* InSelectedAsset, FName InStyleSetName = NAME_None)
+	{
+		SelectedAsset = InSelectedAsset;
+		StyleSetName = InStyleSetName;
+	}
+
+	UPROPERTY(BlueprintReadWrite, Category = "Editor Menus")
+	UObject* SelectedAsset;
+
+	UPROPERTY(BlueprintReadWrite, Category = "Editor Menus")
+	FName StyleSetName;
+
+	TWeakPtr<SMediaPlayerEditorMedia> MediaPlayerEditorMedia;
+};
diff --git a/Engine/Plugins/Media/MediaPlayerEditor/Source/MediaPlayerEditor/Private/Widgets/SMediaPlayerEditorMedia.cpp b/Engine/Plugins/Media/MediaPlayerEditor/Source/MediaPlayerEditor/Private/Widgets/SMediaPlayerEditorMedia.cpp
--- a/Engine/Plugins/Media/MediaPlayerEditor/Source/MediaPlayerEditor/Private/Widgets/SMediaPlayerEditorMedia.cpp
+++ b/Engine/Plugins/Media/MediaPlayerEditor/Source/MediaPlayerEditor/Private/Widgets/SMediaPlayerEditorMedia.cpp
@@ -20,10 +20,12 @@
 #include "Widgets/Layout/SBorder.h"
 #include "Widgets/Notifications/SNotificationList.h"
 #include "Widgets/Text/STextBlock.h"
-
+#include "EditorMenuSubsystem.h"
+#include "Shared/MediaPlayerEditorMediaContext.h"
 
 #define LOCTEXT_NAMESPACE "SMediaPlayerEditorMedia"
 
+const FName SMediaPlayerEditorMedia::AssetPickerAssetContextMenuName = "MediaPlayer.AssetPickerAssetContextMenu";
 
 /* SMediaPlayerEditorMedia interface
  *****************************************************************************/
@@ -124,93 +126,123 @@ void SMediaPlayerEditorMedia::HandleAssetPickerAssetEnterPressed(const TArray<FA
 	}
 }
 
-
-TSharedPtr<SWidget> SMediaPlayerEditorMedia::HandleAssetPickerGetAssetContextMenu(const TArray<FAssetData>& SelectedAssets)
+void SMediaPlayerEditorMedia::RegisterMenus()
 {
-	if (SelectedAssets.Num() <= 0)
+	UEditorMenu* Menu = UEditorMenuSubsystem::Get()->RegisterMenu(AssetPickerAssetContextMenuName);
 	{
-		return nullptr;
-	}
+		FEditorMenuSection& Section = Menu->AddSection("MediaSection", LOCTEXT("MediaSection", "Media"));
+		Section.AddEntry(FEditorMenuEntry::InitMenuEntry(
+			"Edit",
+			LOCTEXT("EditMenuAction", "Edit..."),
+			LOCTEXT("EditMenuActionTooltip", "Opens the selected asset for edit."),
+			FSlateIcon(FEditorStyle::GetStyleSetName(), "ContentBrowser.AssetActions.Edit"),
+			FEditorMenuExecuteAction::CreateLambda([](const FEditorMenuContext& InContext)
+			{
+				if (UMediaPlayerEditorMediaContext* Context = InContext.Find<UMediaPlayerEditorMediaContext>())
+				{
+					if (Context->SelectedAsset)
+					{
+						FAssetEditorManager::Get().OpenEditorForAsset(Context->SelectedAsset);
+					}
+				}
+			})
+		));
 
-	UObject* SelectedAsset = SelectedAssets[0].GetAsset();
 
-	if (SelectedAsset == nullptr)
-	{
-		return nullptr;
+		Section.AddDynamicEntry("Open", FNewEditorMenuSectionDelegate::CreateLambda([](FEditorMenuSection& InSection)
+		{
+			if (UMediaPlayerEditorMediaContext* Context = InSection.FindContext<UMediaPlayerEditorMediaContext>())
+			{
+				InSection.AddEntry(FEditorMenuEntry::InitMenuEntry(
+					"Open",
+					LOCTEXT("OpenMenuAction", "Open"),
+					LOCTEXT("OpenMenuActionTooltip", "Open this media asset in the player"),
+					FSlateIcon(Context->StyleSetName, "MediaPlayerEditor.NextMedia.Small"),
+					FEditorMenuExecuteAction::CreateLambda([=](const FEditorMenuContext& InContext)
+					{
+						if (UMediaPlayerEditorMediaContext* Context = InContext.Find<UMediaPlayerEditorMediaContext>())
+						{
+							if (Context->MediaPlayerEditorMedia.IsValid())
+							{
+								Context->MediaPlayerEditorMedia.Pin()->OpenMediaAsset(Context->SelectedAsset);
+							}
+						}
+					})
+				));
+			}
+		}));
 	}
 
-	FMenuBuilder MenuBuilder(true /*bInShouldCloseWindowAfterMenuSelection=*/, nullptr);
 	{
-		// media section
-		MenuBuilder.BeginSection("MediaSection", LOCTEXT("MediaSection", "Media"));
-		{
-			MenuBuilder.AddMenuEntry(
-				LOCTEXT("EditMenuAction", "Edit..."),
-				LOCTEXT("EditMenuActionTooltip", "Opens the selected asset for edit."),
-				FSlateIcon(FEditorStyle::GetStyleSetName(), "ContentBrowser.AssetActions.Edit"),
-				FUIAction(
-					FExecuteAction::CreateLambda([=]() {
-						FAssetEditorManager::Get().OpenEditorForAsset(SelectedAsset);
-					})
-				)
-			);
-
-			MenuBuilder.AddMenuEntry(
-				LOCTEXT("OpenMenuAction", "Open"),
-				LOCTEXT("OpenMenuActionTooltip", "Open this media asset in the player"),
-				FSlateIcon(Style->GetStyleSetName(), "MediaPlayerEditor.NextMedia.Small"),
-				FUIAction(
-					FExecuteAction::CreateLambda([=]() {
-						OpenMediaAsset(SelectedAsset);
-					})
-				)
-			);
-		}
-		MenuBuilder.EndSection();
+		FEditorMenuSection& Section = Menu->AddSection("AssetSection", LOCTEXT("AssetSection", "Asset"));
+		Section.AddEntry(FEditorMenuEntry::InitMenuEntry(
+			"FindInCbMenuAction",
+			LOCTEXT("FindInCbMenuAction", "Find in Content Browser"),
+			LOCTEXT("FindInCbMenuActionTooltip", "Summons the Content Browser and navigates to the selected asset"),
+			FSlateIcon(FEditorStyle::GetStyleSetName(), "SystemWideCommands.FindInContentBrowser"),
+			FEditorMenuExecuteAction::CreateLambda([](const FEditorMenuContext& InContext)
+			{
+				if (UMediaPlayerEditorMediaContext* Context = InContext.Find<UMediaPlayerEditorMediaContext>())
+				{
+					TArray<UObject*> AssetsToSync = TArrayBuilder<UObject*>().Add(Context->SelectedAsset);
+					GEditor->SyncBrowserToObjects(AssetsToSync);
+				}
+			})
+		));
 
-		// asset section
-		MenuBuilder.BeginSection("AssetSection", LOCTEXT("AssetSection", "Asset"));
+		Section.AddDynamicEntry("OpenInFileManager", FNewEditorMenuSectionDelegate::CreateLambda([](FEditorMenuSection& InSection)
 		{
-			MenuBuilder.AddMenuEntry(
-				LOCTEXT("FindInCbMenuAction", "Find in Content Browser"),
-				LOCTEXT("FindInCbMenuActionTooltip", "Summons the Content Browser and navigates to the selected asset"),
-				FSlateIcon(FEditorStyle::GetStyleSetName(), "SystemWideCommands.FindInContentBrowser"),
-				FUIAction(
-					FExecuteAction::CreateLambda([=]() {
-						TArray<UObject*> AssetsToSync = TArrayBuilder<UObject*>().Add(SelectedAsset);
-						GEditor->SyncBrowserToObjects(AssetsToSync);
-					})
-				)
-			);
-
-			auto FileMediaSource = Cast<UFileMediaSource>(SelectedAsset);
-
-			if (FileMediaSource != nullptr)
+			if (UMediaPlayerEditorMediaContext* Context = InSection.FindContext<UMediaPlayerEditorMediaContext>())
 			{
-				FFormatNamedArguments Args;
+				if (UFileMediaSource* FileMediaSource = Cast<UFileMediaSource>(Context->SelectedAsset))
 				{
-					Args.Add(TEXT("FileManagerName"), FPlatformMisc::GetFileManagerName());
+					FFormatNamedArguments Args;
+					{
+						Args.Add(TEXT("FileManagerName"), FPlatformMisc::GetFileManagerName());
+					}
+
+					TWeakObjectPtr<UFileMediaSource> FileMediaSourceWeak = FileMediaSource;
+					InSection.AddEntry(FEditorMenuEntry::InitMenuEntry(
+						"OpenInFileManager",
+						FText::Format(LOCTEXT("OpenInFileManager", "Show Media File in {FileManagerName}"), Args),
+						LOCTEXT("OpenInFileManagerTooltip", "Finds the media file that this asset points to on disk"),
+						FSlateIcon(FEditorStyle::GetStyleSetName(), "SystemWideCommands.FindInContentBrowser"),
+						FUIAction(
+							FExecuteAction::CreateLambda([=]() {
+								if (FileMediaSourceWeak.IsValid())
+								{
+									FPlatformProcess::ExploreFolder(*FileMediaSourceWeak->GetFullPath());
+								}
+							}),
+							FCanExecuteAction::CreateLambda([=]() -> bool {
+								return FileMediaSourceWeak.IsValid() && FileMediaSourceWeak->Validate();
+							})
+						)
+					));
 				}
-
-				MenuBuilder.AddMenuEntry(
-					FText::Format(LOCTEXT("OpenInFileManager", "Show Media File in {FileManagerName}"), Args),
-					LOCTEXT("OpenInFileManagerTooltip", "Finds the media file that this asset points to on disk"),
-					FSlateIcon(FEditorStyle::GetStyleSetName(), "SystemWideCommands.FindInContentBrowser"),
-					FUIAction(
-						FExecuteAction::CreateLambda([=]() {
-							FPlatformProcess::ExploreFolder(*FileMediaSource->GetFullPath());
-						}),
-						FCanExecuteAction::CreateLambda([=]() -> bool {
-							return FileMediaSource->Validate();
-						})
-					)
-				);
 			}
-		}
-		MenuBuilder.EndSection();
+		}));
+	}
+}
+
+TSharedPtr<SWidget> SMediaPlayerEditorMedia::HandleAssetPickerGetAssetContextMenu(const TArray<FAssetData>& SelectedAssets)
+{
+	if (SelectedAssets.Num() <= 0)
+	{
+		return nullptr;
+	}
+
+	UObject* SelectedAsset = SelectedAssets[0].GetAsset();
+	if (SelectedAsset == nullptr)
+	{
+		return nullptr;
 	}
 
-	return MenuBuilder.MakeWidget();
+	UMediaPlayerEditorMediaContext* ContextObject = NewObject<UMediaPlayerEditorMediaContext>();
+	ContextObject->InitContext(SelectedAsset, Style->GetStyleSetName());
+	ContextObject->MediaPlayerEditorMedia = SharedThis(this);
+	FEditorMenuContext MenuContext(ContextObject);
+	return UEditorMenuSubsystem::Get()->GenerateWidget(AssetPickerAssetContextMenuName, MenuContext);
 }
 
 
diff --git a/Engine/Plugins/Media/MediaPlayerEditor/Source/MediaPlayerEditor/Private/Widgets/SMediaPlayerEditorMedia.h b/Engine/Plugins/Media/MediaPlayerEditor/Source/MediaPlayerEditor/Private/Widgets/SMediaPlayerEditorMedia.h
--- a/Engine/Plugins/Media/MediaPlayerEditor/Source/MediaPlayerEditor/Private/Widgets/SMediaPlayerEditorMedia.h
+++ b/Engine/Plugins/Media/MediaPlayerEditor/Source/MediaPlayerEditor/Private/Widgets/SMediaPlayerEditorMedia.h
@@ -34,7 +34,8 @@ public:
 	 */
 	void Construct(const FArguments& InArgs, UMediaPlayer& InMediaPlayer, const TSharedRef<ISlateStyle>& InStyle);
 
-protected:
+	/* Register menus associated with this widget */
+	static void RegisterMenus();
 
 	/**
 	 * Open an asset in the media player.
@@ -43,8 +44,9 @@ protected:
 	 */
 	void OpenMediaAsset(UObject* Asset);
 
+protected:
 	/** Show a message box for media opening failures. */
-	void ShowMediaOpenFailedMessage();
+	static void ShowMediaOpenFailedMessage();
 
 private:
 
@@ -64,4 +66,7 @@ private:
 
 	/** The widget style set to use. */
 	TSharedPtr<ISlateStyle> Style;
+
+	/** Name of the asset picker asset context menu */
+	static const FName AssetPickerAssetContextMenuName;
 };
