Rename Editor\EditorMenus\ module to Runtime\ToolMenus\ so that other modules can use the API without having to depend on editor modules.
#rnx

diff --git a/Engine/Plugins/2D/Paper2D/Source/Paper2DEditor/Paper2DEditor.Build.cs b/Engine/Plugins/2D/Paper2D/Source/Paper2DEditor/Paper2DEditor.Build.cs
--- a/Engine/Plugins/2D/Paper2D/Source/Paper2DEditor/Paper2DEditor.Build.cs
+++ b/Engine/Plugins/2D/Paper2D/Source/Paper2DEditor/Paper2DEditor.Build.cs
@@ -31,7 +31,7 @@ public class Paper2DEditor : ModuleRules
 				"EditorWidgets",
 				"Projects",
 				"NavigationSystem",
-				"EditorMenus",
+				"ToolMenus",
 			});
 
 		PrivateIncludePathModuleNames.AddRange(
diff --git a/Engine/Plugins/2D/Paper2D/Source/Paper2DEditor/Private/SpriteAssetTypeActions.cpp b/Engine/Plugins/2D/Paper2D/Source/Paper2DEditor/Private/SpriteAssetTypeActions.cpp
--- a/Engine/Plugins/2D/Paper2D/Source/Paper2DEditor/Private/SpriteAssetTypeActions.cpp
+++ b/Engine/Plugins/2D/Paper2D/Source/Paper2DEditor/Private/SpriteAssetTypeActions.cpp
@@ -2,7 +2,7 @@
 
 #include "SpriteAssetTypeActions.h"
 #include "Framework/MultiBox/MultiBoxBuilder.h"
-#include "EditorMenuSubsystem.h"
+#include "ToolMenus.h"
 #include "Misc/PackageName.h"
 #include "Misc/FeedbackContext.h"
 #include "EditorStyleSet.h"
@@ -59,7 +59,7 @@ uint32 FSpriteAssetTypeActions::GetCategories()
 	return MyAssetCategory;
 }
 
-void FSpriteAssetTypeActions::GetActions(const TArray<UObject*>& InObjects, FEditorMenuSection& Section)
+void FSpriteAssetTypeActions::GetActions(const TArray<UObject*>& InObjects, FToolMenuSection& Section)
 {
 	auto Sprites = GetTypedWeakObjectPtrs<UPaperSprite>(InObjects);
 
diff --git a/Engine/Plugins/2D/Paper2D/Source/Paper2DEditor/Private/SpriteAssetTypeActions.h b/Engine/Plugins/2D/Paper2D/Source/Paper2DEditor/Private/SpriteAssetTypeActions.h
--- a/Engine/Plugins/2D/Paper2D/Source/Paper2DEditor/Private/SpriteAssetTypeActions.h
+++ b/Engine/Plugins/2D/Paper2D/Source/Paper2DEditor/Private/SpriteAssetTypeActions.h
@@ -18,7 +18,7 @@ public:
 	virtual FColor GetTypeColor() const override;
 	virtual UClass* GetSupportedClass() const override;
 	virtual bool HasActions(const TArray<UObject*>& InObjects) const override { return true; }
-	virtual void GetActions(const TArray<UObject*>& InObjects, struct FEditorMenuSection& Section) override;
+	virtual void GetActions(const TArray<UObject*>& InObjects, struct FToolMenuSection& Section) override;
 	virtual void OpenAssetEditor(const TArray<UObject*>& InObjects, TSharedPtr<class IToolkitHost> EditWithinLevelEditor = TSharedPtr<IToolkitHost>()) override;
 	virtual uint32 GetCategories() override;
 	// End of IAssetTypeActions interface
diff --git a/Engine/Plugins/2D/Paper2D/Source/Paper2DEditor/Private/TileSetAssetTypeActions.cpp b/Engine/Plugins/2D/Paper2D/Source/Paper2DEditor/Private/TileSetAssetTypeActions.cpp
--- a/Engine/Plugins/2D/Paper2D/Source/Paper2DEditor/Private/TileSetAssetTypeActions.cpp
+++ b/Engine/Plugins/2D/Paper2D/Source/Paper2DEditor/Private/TileSetAssetTypeActions.cpp
@@ -1,7 +1,7 @@
 // Copyright 1998-2019 Epic Games, Inc. All Rights Reserved.
 
 #include "TileSetAssetTypeActions.h"
-#include "EditorMenuSubsystem.h"
+#include "ToolMenus.h"
 #include "Engine/Texture.h"
 #include "Misc/PackageName.h"
 #include "EditorStyleSet.h"
@@ -58,7 +58,7 @@ uint32 FTileSetAssetTypeActions::GetCategories()
 	return MyAssetCategory;
 }
 
-void FTileSetAssetTypeActions::GetActions(const TArray<UObject*>& InObjects, FEditorMenuSection& Section)
+void FTileSetAssetTypeActions::GetActions(const TArray<UObject*>& InObjects, FToolMenuSection& Section)
 {
 	TArray<TWeakObjectPtr<UPaperTileSet>> TileSets = GetTypedWeakObjectPtrs<UPaperTileSet>(InObjects);
 
diff --git a/Engine/Plugins/2D/Paper2D/Source/Paper2DEditor/Private/TileSetAssetTypeActions.h b/Engine/Plugins/2D/Paper2D/Source/Paper2DEditor/Private/TileSetAssetTypeActions.h
--- a/Engine/Plugins/2D/Paper2D/Source/Paper2DEditor/Private/TileSetAssetTypeActions.h
+++ b/Engine/Plugins/2D/Paper2D/Source/Paper2DEditor/Private/TileSetAssetTypeActions.h
@@ -20,7 +20,7 @@ public:
 	virtual void OpenAssetEditor(const TArray<UObject*>& InObjects, TSharedPtr<class IToolkitHost> EditWithinLevelEditor = TSharedPtr<IToolkitHost>()) override;
 	virtual uint32 GetCategories() override;
 	virtual bool HasActions(const TArray<UObject*>& InObjects) const override { return true; }
-	virtual void GetActions(const TArray<UObject*>& InObjects, struct FEditorMenuSection& Section) override;
+	virtual void GetActions(const TArray<UObject*>& InObjects, struct FToolMenuSection& Section) override;
 	// End of IAssetTypeActions interface
 
 private:
diff --git a/Engine/Plugins/2D/Paper2D/Source/PaperSpriteSheetImporter/PaperSpriteSheetImporter.Build.cs b/Engine/Plugins/2D/Paper2D/Source/PaperSpriteSheetImporter/PaperSpriteSheetImporter.Build.cs
--- a/Engine/Plugins/2D/Paper2D/Source/PaperSpriteSheetImporter/PaperSpriteSheetImporter.Build.cs
+++ b/Engine/Plugins/2D/Paper2D/Source/PaperSpriteSheetImporter/PaperSpriteSheetImporter.Build.cs
@@ -20,7 +20,7 @@ public class PaperSpriteSheetImporter : ModuleRules
 				"AssetTools",
 				"ContentBrowser",
                 "EditorStyle",
-				"EditorMenus",
+				"ToolMenus",
 			});
 
 		PrivateIncludePathModuleNames.AddRange(
diff --git a/Engine/Plugins/2D/Paper2D/Source/PaperSpriteSheetImporter/Private/PaperSpriteSheetAssetTypeActions.cpp b/Engine/Plugins/2D/Paper2D/Source/PaperSpriteSheetImporter/Private/PaperSpriteSheetAssetTypeActions.cpp
--- a/Engine/Plugins/2D/Paper2D/Source/PaperSpriteSheetImporter/Private/PaperSpriteSheetAssetTypeActions.cpp
+++ b/Engine/Plugins/2D/Paper2D/Source/PaperSpriteSheetImporter/Private/PaperSpriteSheetAssetTypeActions.cpp
@@ -2,7 +2,7 @@
 
 #include "PaperSpriteSheetAssetTypeActions.h"
 #include "Misc/PackageName.h"
-#include "EditorMenuSubsystem.h"
+#include "ToolMenus.h"
 #include "Misc/FeedbackContext.h"
 #include "EditorStyleSet.h"
 #include "EditorFramework/AssetImportData.h"
@@ -57,7 +57,7 @@ void FPaperSpriteSheetAssetTypeActions::GetResolvedSourceFilePaths(const TArray<
 	}
 }
 
-void FPaperSpriteSheetAssetTypeActions::GetActions(const TArray<UObject*>& InObjects, FEditorMenuSection& Section)
+void FPaperSpriteSheetAssetTypeActions::GetActions(const TArray<UObject*>& InObjects, FToolMenuSection& Section)
 {
 	auto SpriteSheetImports = GetTypedWeakObjectPtrs<UPaperSpriteSheet>(InObjects);
 
diff --git a/Engine/Plugins/Editor/DataValidation/Source/DataValidation/Private/DataValidationModule.cpp b/Engine/Plugins/Editor/DataValidation/Source/DataValidation/Private/DataValidationModule.cpp
--- a/Engine/Plugins/Editor/DataValidation/Source/DataValidation/Private/DataValidationModule.cpp
+++ b/Engine/Plugins/Editor/DataValidation/Source/DataValidation/Private/DataValidationModule.cpp
@@ -6,7 +6,7 @@
 #include "GameFramework/HUD.h"
 
 #include "Framework/Application/SlateApplication.h"
-#include "EditorMenuSubsystem.h"
+#include "ToolMenus.h"
 #include "EditorStyleSet.h"
 #include "AssetRegistryModule.h"
 #include "Misc/MessageDialog.h"
@@ -94,7 +94,7 @@ void FDataValidationModule::ShutdownModule()
 		}
 
 		// remove menu extension
-		UEditorMenuSubsystem::UnregisterOwner(this);
+		UToolMenus::UnregisterOwner(this);
 		FCoreDelegates::OnPostEngineInit.RemoveAll(this);
 
 		UPackage::PackageSavedEvent.RemoveAll(this);
@@ -103,10 +103,10 @@ void FDataValidationModule::ShutdownModule()
 
 void FDataValidationModule::RegisterMenus()
 {
-	FEditorMenuOwnerScoped OwnerScoped(this);
-	UEditorMenu* Menu = UEditorMenuSubsystem::Get()->ExtendMenu("LevelEditor.MainMenu.File");
-	FEditorMenuSection& Section = Menu->AddSection("DataValidation", LOCTEXT("DataValidation", "DataValidation"), FEditorMenuInsert("FileLoadAndSave", EEditorMenuInsertType::After));
-	Section.AddEntry(FEditorMenuEntry::InitMenuEntry(
+	FToolMenuOwnerScoped OwnerScoped(this);
+	UToolMenu* Menu = UToolMenus::Get()->ExtendMenu("LevelEditor.MainMenu.File");
+	FToolMenuSection& Section = Menu->AddSection("DataValidation", LOCTEXT("DataValidation", "DataValidation"), FToolMenuInsert("FileLoadAndSave", EToolMenuInsertType::After));
+	Section.AddEntry(FToolMenuEntry::InitMenuEntry(
 		"ValidateData",
 		TAttribute<FText>::Create(&Menu_ValidateDataGetTitle),
 		LOCTEXT("ValidateDataTooltip", "Validates all user data in content directory."),
diff --git a/Engine/Plugins/Editor/GameplayTagsEditor/Source/GameplayTagsEditor/Private/AssetTypeActions_GameplayTagAssetBase.cpp b/Engine/Plugins/Editor/GameplayTagsEditor/Source/GameplayTagsEditor/Private/AssetTypeActions_GameplayTagAssetBase.cpp
--- a/Engine/Plugins/Editor/GameplayTagsEditor/Source/GameplayTagsEditor/Private/AssetTypeActions_GameplayTagAssetBase.cpp
+++ b/Engine/Plugins/Editor/GameplayTagsEditor/Source/GameplayTagsEditor/Private/AssetTypeActions_GameplayTagAssetBase.cpp
@@ -1,7 +1,7 @@
 // Copyright 1998-2019 Epic Games, Inc. All Rights Reserved.
 
 #include "AssetTypeActions_GameplayTagAssetBase.h"
-#include "EditorMenuSubsystem.h"
+#include "ToolMenus.h"
 #include "UObject/UnrealType.h"
 #include "Framework/Application/SlateApplication.h"
 
@@ -20,7 +20,7 @@ bool FAssetTypeActions_GameplayTagAssetBase::HasActions(const TArray<UObject*>&
 	return true;
 }
 
-void FAssetTypeActions_GameplayTagAssetBase::GetActions(const TArray<UObject*>& InObjects, FEditorMenuSection& Section)
+void FAssetTypeActions_GameplayTagAssetBase::GetActions(const TArray<UObject*>& InObjects, FToolMenuSection& Section)
 {
 	TArray<UObject*> ContainerObjectOwners;
 	TArray<FGameplayTagContainer*> Containers;
diff --git a/Engine/Plugins/Experimental/ChaosSolverPlugin/Source/ChaosSolverEditor/Private/Chaos/AssetTypeActions_ChaosSolver.cpp b/Engine/Plugins/Experimental/ChaosSolverPlugin/Source/ChaosSolverEditor/Private/Chaos/AssetTypeActions_ChaosSolver.cpp
--- a/Engine/Plugins/Experimental/ChaosSolverPlugin/Source/ChaosSolverEditor/Private/Chaos/AssetTypeActions_ChaosSolver.cpp
+++ b/Engine/Plugins/Experimental/ChaosSolverPlugin/Source/ChaosSolverEditor/Private/Chaos/AssetTypeActions_ChaosSolver.cpp
@@ -4,7 +4,7 @@
 
 #include "ThumbnailRendering/SceneThumbnailInfo.h"
 #include "Chaos/ChaosSolver.h"
-#include "EditorMenuSubsystem.h"
+#include "ToolMenus.h"
 
 #define LOCTEXT_NAMESPACE "AssetTypeActions"
 
@@ -19,7 +19,7 @@ UThumbnailInfo* FAssetTypeActions_ChaosSolver::GetThumbnailInfo(UObject* Asset)
 	return NewObject<USceneThumbnailInfo>(ChaosSolver, NAME_None, RF_Transactional);
 }
 
-void FAssetTypeActions_ChaosSolver::GetActions(const TArray<UObject*>& InObjects, FEditorMenuSection& Section)
+void FAssetTypeActions_ChaosSolver::GetActions(const TArray<UObject*>& InObjects, FToolMenuSection& Section)
 {
 	FAssetTypeActions_Base::GetActions(InObjects, Section);
 }
diff --git a/Engine/Plugins/Experimental/CodeEditor/Source/CodeEditor/Private/CodeEditor.cpp b/Engine/Plugins/Experimental/CodeEditor/Source/CodeEditor/Private/CodeEditor.cpp
--- a/Engine/Plugins/Experimental/CodeEditor/Source/CodeEditor/Private/CodeEditor.cpp
+++ b/Engine/Plugins/Experimental/CodeEditor/Source/CodeEditor/Private/CodeEditor.cpp
@@ -5,7 +5,7 @@
 #include "Modules/ModuleManager.h"
 #include "Textures/SlateIcon.h"
 #include "Framework/Commands/UIAction.h"
-#include "EditorMenuSubsystem.h"
+#include "ToolMenus.h"
 #include "Framework/Docking/TabManager.h"
 #include "Toolkits/IToolkitHost.h"
 #include "CodeEditorStyle.h"
@@ -37,14 +37,14 @@ public:
 
 	virtual void OnPostEngineInit()
 	{
-		if (UEditorMenuSubsystem::IsRunningEditorUI())
+		if (UToolMenus::IsToolMenuUIEnabled())
 		{
-			UEditorMenu* Menu = UEditorMenuSubsystem::Get()->ExtendMenu("LevelEditor.MainMenu.File");
-			FEditorMenuSection& Section = Menu->FindOrAddSection("FileProject");
+			UToolMenu* Menu = UToolMenus::Get()->ExtendMenu("LevelEditor.MainMenu.File");
+			FToolMenuSection& Section = Menu->FindOrAddSection("FileProject");
 
-			FEditorMenuOwnerScoped OwnerScoped(this);
+			FToolMenuOwnerScoped OwnerScoped(this);
 			{
-				FEditorMenuEntry& MenuEntry = Section.AddMenuEntry(
+				FToolMenuEntry& MenuEntry = Section.AddMenuEntry(
 					"EditSourceCode",
 					LOCTEXT("CodeEditorTabTitle", "Edit Source Code"),
 					LOCTEXT("CodeEditorTooltipText", "Open the Code Editor tab."),
@@ -54,7 +54,7 @@ public:
 						FExecuteAction::CreateStatic(&FCodeEditor::OpenCodeEditor)
 					)
 				);
-				MenuEntry.InsertPosition = FEditorMenuInsert(NAME_None, EEditorMenuInsertType::First);
+				MenuEntry.InsertPosition = FToolMenuInsert(NAME_None, EToolMenuInsertType::First);
 			}
 		}
 	}
@@ -79,7 +79,7 @@ public:
 		FGlobalTabmanager::Get()->UnregisterTabSpawner( CodeEditorTabName );
 
 		FCoreDelegates::OnPostEngineInit.RemoveAll(this);
-		UEditorMenuSubsystem::UnregisterOwner(this);
+		UToolMenus::UnregisterOwner(this);
 
 		FCodeEditorStyle::Shutdown();
 	}
diff --git a/Engine/Plugins/Experimental/FieldSystemPlugin/Source/FieldSyStemEditor/Private/Field/AssetTypeActions_FieldSystem.cpp b/Engine/Plugins/Experimental/FieldSystemPlugin/Source/FieldSyStemEditor/Private/Field/AssetTypeActions_FieldSystem.cpp
--- a/Engine/Plugins/Experimental/FieldSystemPlugin/Source/FieldSyStemEditor/Private/Field/AssetTypeActions_FieldSystem.cpp
+++ b/Engine/Plugins/Experimental/FieldSystemPlugin/Source/FieldSyStemEditor/Private/Field/AssetTypeActions_FieldSystem.cpp
@@ -5,7 +5,7 @@
 #include "ThumbnailRendering/SceneThumbnailInfo.h"
 #include "Field/FieldSystem.h"
 #include "Field/FieldSystemAsset.h"
-#include "EditorMenuSubsystem.h"
+#include "ToolMenus.h"
 
 #define LOCTEXT_NAMESPACE "AssetTypeActions"
 
@@ -20,7 +20,7 @@ UThumbnailInfo* FAssetTypeActions_FieldSystem::GetThumbnailInfo(UObject* Asset)
 	return NewObject<USceneThumbnailInfo>(FieldSystem, NAME_None, RF_Transactional);
 }
 
-void FAssetTypeActions_FieldSystem::GetActions(const TArray<UObject*>& InObjects, FEditorMenuSection& Section)
+void FAssetTypeActions_FieldSystem::GetActions(const TArray<UObject*>& InObjects, FToolMenuSection& Section)
 {
 	FAssetTypeActions_Base::GetActions(InObjects, Section);
 }
diff --git a/Engine/Plugins/Experimental/GeometryCache/Source/GeometryCacheEd/Private/AssetTypeActions_GeometryCache.cpp b/Engine/Plugins/Experimental/GeometryCache/Source/GeometryCacheEd/Private/AssetTypeActions_GeometryCache.cpp
--- a/Engine/Plugins/Experimental/GeometryCache/Source/GeometryCacheEd/Private/AssetTypeActions_GeometryCache.cpp
+++ b/Engine/Plugins/Experimental/GeometryCache/Source/GeometryCacheEd/Private/AssetTypeActions_GeometryCache.cpp
@@ -4,7 +4,7 @@
 #include "EditorFramework/AssetImportData.h"
 #include "ThumbnailRendering/SceneThumbnailInfo.h"
 #include "GeometryCache.h"
-#include "EditorMenuSubsystem.h"
+#include "ToolMenus.h"
 
 FText FAssetTypeActions_GeometryCache::GetName() const
 {
@@ -26,7 +26,7 @@ bool FAssetTypeActions_GeometryCache::HasActions(const TArray<UObject*>& InObjec
 	return false;
 }
 
-void FAssetTypeActions_GeometryCache::GetActions(const TArray<UObject*>& InObjects, FEditorMenuSection& Section)
+void FAssetTypeActions_GeometryCache::GetActions(const TArray<UObject*>& InObjects, FToolMenuSection& Section)
 {
 	FAssetTypeActions_Base::GetActions(InObjects, Section);
 }
diff --git a/Engine/Plugins/Experimental/GeometryCollectionPlugin/Source/GeometryCollectionEditor/Private/GeometryCollection/AssetTypeActions_GeometryCollection.cpp b/Engine/Plugins/Experimental/GeometryCollectionPlugin/Source/GeometryCollectionEditor/Private/GeometryCollection/AssetTypeActions_GeometryCollection.cpp
--- a/Engine/Plugins/Experimental/GeometryCollectionPlugin/Source/GeometryCollectionEditor/Private/GeometryCollection/AssetTypeActions_GeometryCollection.cpp
+++ b/Engine/Plugins/Experimental/GeometryCollectionPlugin/Source/GeometryCollectionEditor/Private/GeometryCollection/AssetTypeActions_GeometryCollection.cpp
@@ -3,7 +3,7 @@
 #include "GeometryCollection/AssetTypeActions_GeometryCollection.h"
 #include "ThumbnailRendering/SceneThumbnailInfo.h"
 #include "GeometryCollection/GeometryCollectionObject.h"
-#include "EditorMenuSubsystem.h"
+#include "ToolMenus.h"
 
 #define LOCTEXT_NAMESPACE "AssetTypeActions"
 
@@ -25,7 +25,7 @@ UThumbnailInfo* FAssetTypeActions_GeometryCollection::GetThumbnailInfo(UObject*
 	return ThumbnailInfo;
 }
 
-void FAssetTypeActions_GeometryCollection::GetActions(const TArray<UObject*>& InObjects, FEditorMenuSection& Section)
+void FAssetTypeActions_GeometryCollection::GetActions(const TArray<UObject*>& InObjects, FToolMenuSection& Section)
 {
 	FAssetTypeActions_Base::GetActions(InObjects, Section);
 	//Set("ClassIcon.GeometryCollection", new IMAGE_BRUSH("Icons/AssetIcons/GeometryCollectionComponentAtlasGroup_16x", Icon16x16));
diff --git a/Engine/Plugins/Experimental/PythonScriptPlugin/Source/PythonScriptPlugin/Private/PythonScriptPlugin.cpp b/Engine/Plugins/Experimental/PythonScriptPlugin/Source/PythonScriptPlugin/Private/PythonScriptPlugin.cpp
--- a/Engine/Plugins/Experimental/PythonScriptPlugin/Source/PythonScriptPlugin/Private/PythonScriptPlugin.cpp
+++ b/Engine/Plugins/Experimental/PythonScriptPlugin/Source/PythonScriptPlugin/Private/PythonScriptPlugin.cpp
@@ -32,7 +32,7 @@
 #include "Engine/Engine.h"
 #include "Framework/Application/SlateApplication.h"
 #include "Framework/MultiBox/MultiBoxBuilder.h"
-#include "EditorMenuSubsystem.h"
+#include "ToolMenus.h"
 #include "Misc/ConfigCacheIni.h"
 #include "Misc/MessageDialog.h"
 #endif	// WITH_EDITOR
@@ -224,7 +224,7 @@ public:
 
 	virtual void OnShutdownMenu() override
 	{
-		UEditorMenuSubsystem::UnregisterOwner(this);
+		UToolMenus::UnregisterOwner(this);
 
 		// Write to file
 		if (bRecentsFilesDirty)
@@ -313,10 +313,10 @@ private:
 		GConfig->Flush(false);
 	}
 
-	void MakeRecentPythonScriptMenu(UEditorMenu* InMenu)
+	void MakeRecentPythonScriptMenu(UToolMenu* InMenu)
 	{
-		FEditorMenuOwnerScoped OwnerScoped(this);
-		FEditorMenuSection& Section = InMenu->AddSection("Files");
+		FToolMenuOwnerScoped OwnerScoped(this);
+		FToolMenuSection& Section = InMenu->AddSection("Files");
 		for (int32 Index = RecentsFiles.Num() - 1; Index >= 0; --Index)
 		{
 			Section.AddMenuEntry(
@@ -328,7 +328,7 @@ private:
 			);
 		}
 
-		FEditorMenuSection& ClearSection = InMenu->AddSection("Clear");
+		FToolMenuSection& ClearSection = InMenu->AddSection("Clear");
 		Section.AddMenuEntry(
 			"ClearRecentPython",
 			LOCTEXT("ClearRecentPython", "Clear Recent Python Scripts"),
@@ -340,9 +340,9 @@ private:
 
 	void RegisterMenus()
 	{
-		FEditorMenuOwnerScoped OwnerScoped(this);
-		UEditorMenu* Menu = UEditorMenuSubsystem::Get()->ExtendMenu("LevelEditor.MainMenu.File");
-		FEditorMenuSection& Section = Menu->AddSection("Python", LOCTEXT("Python", "Python"), FEditorMenuInsert("FileLoadAndSave", EEditorMenuInsertType::After));
+		FToolMenuOwnerScoped OwnerScoped(this);
+		UToolMenu* Menu = UToolMenus::Get()->ExtendMenu("LevelEditor.MainMenu.File");
+		FToolMenuSection& Section = Menu->AddSection("Python", LOCTEXT("Python", "Python"), FToolMenuInsert("FileLoadAndSave", EToolMenuInsertType::After));
 		Section.AddMenuEntry(
 			"OpenPython",
 			LOCTEXT("OpenPython", "Execute Python Script"),
@@ -350,12 +350,12 @@ private:
 			FSlateIcon(),
 			FUIAction(FExecuteAction::CreateRaw(this, &FPythonCommandMenuImpl::Menu_ExecutePython))
 		);
-		Section.AddEntry(FEditorMenuEntry::InitSubMenu(
+		Section.AddEntry(FToolMenuEntry::InitSubMenu(
 			"LevelEditor.MainMenu.File",
 			"RecentPythonsSubMenu",
 			LOCTEXT("RecentPythonsSubMenu", "Recent Python Scripts"),
 			LOCTEXT("RecentPythonsSubMenu_ToolTip", "Select a recent Python Script file and Execute it."),
-			FNewEditorMenuDelegate::CreateRaw(this, &FPythonCommandMenuImpl::MakeRecentPythonScriptMenu),
+			FNewToolMenuDelegate::CreateRaw(this, &FPythonCommandMenuImpl::MakeRecentPythonScriptMenu),
 			false,
 			FSlateIcon(FEditorStyle::GetStyleSetName(), "MainFrame.RecentLevels")
 		));
@@ -510,13 +510,13 @@ void FPythonScriptPlugin::StartupModule()
 #if WITH_EDITOR
 void FPythonScriptPlugin::OnPostEngineInit()
 {
-	if (UEditorMenuSubsystem::IsRunningEditorUI())
+	if (UToolMenus::IsToolMenuUIEnabled())
 	{
 		check(CmdMenu == nullptr);
 		CmdMenu = new FPythonCommandMenuImpl();
 		CmdMenu->OnStartupMenu();
 
-		UEditorMenuSubsystem::Get()->RegisterStringCommandHandler("Python", FEditorMenuExecuteString::CreateLambda([this](const FString& InString, const FEditorMenuContext& InContext) {
+		UToolMenus::Get()->RegisterStringCommandHandler("Python", FToolMenuExecuteString::CreateLambda([this](const FString& InString, const FToolMenuContext& InContext) {
 			ExecPythonCommand(*InString);
 		}));
 	}
@@ -542,9 +542,9 @@ void FPythonScriptPlugin::ShutdownModule()
 		CmdMenu = nullptr;
 	}
 
-	if (UEditorMenuSubsystem* EditorMenus = UEditorMenuSubsystem::TryGet())
+	if (UToolMenus* ToolMenus = UToolMenus::TryGet())
 	{
-		EditorMenus->UnregisterStringCommandHandler("Python");
+		ToolMenus->UnregisterStringCommandHandler("Python");
 	}
 #endif // WITH_EDITOR
 
diff --git a/Engine/Plugins/Experimental/PythonScriptPlugin/Source/PythonScriptPlugin/PythonScriptPlugin.Build.cs b/Engine/Plugins/Experimental/PythonScriptPlugin/Source/PythonScriptPlugin/PythonScriptPlugin.Build.cs
--- a/Engine/Plugins/Experimental/PythonScriptPlugin/Source/PythonScriptPlugin/PythonScriptPlugin.Build.cs
+++ b/Engine/Plugins/Experimental/PythonScriptPlugin/Source/PythonScriptPlugin/PythonScriptPlugin.Build.cs
@@ -33,7 +33,7 @@ namespace UnrealBuildTool.Rules
 			{
 				PublicDependencyModuleNames.AddRange(
 					new string[] {
-						"EditorMenus"
+						"ToolMenus"
 					}
 				);
 
diff --git a/Engine/Plugins/FX/HoudiniNiagara/Source/HoudiniNiagaraEditor/HoudiniNiagaraEditor.Build.cs b/Engine/Plugins/FX/HoudiniNiagara/Source/HoudiniNiagaraEditor/HoudiniNiagaraEditor.Build.cs
--- a/Engine/Plugins/FX/HoudiniNiagara/Source/HoudiniNiagaraEditor/HoudiniNiagaraEditor.Build.cs
+++ b/Engine/Plugins/FX/HoudiniNiagara/Source/HoudiniNiagaraEditor/HoudiniNiagaraEditor.Build.cs
@@ -53,7 +53,6 @@ public class HoudiniNiagaraEditor : ModuleRules
                 "HoudiniNiagara",
                 "UnrealEd",
 				"EditorStyle",
-				"EditorMenus",
 			}
         );
 			
@@ -69,7 +68,7 @@ public class HoudiniNiagaraEditor : ModuleRules
                 "NiagaraShader",
                 "UtilityShaders",
                 "HoudiniNiagara",
-				"EditorMenus",
+				"ToolMenus",
 			}
         );
 
diff --git a/Engine/Plugins/FX/HoudiniNiagara/Source/HoudiniNiagaraEditor/Private/HoudiniCSVAssetActions.cpp b/Engine/Plugins/FX/HoudiniNiagara/Source/HoudiniNiagaraEditor/Private/HoudiniCSVAssetActions.cpp
--- a/Engine/Plugins/FX/HoudiniNiagara/Source/HoudiniNiagaraEditor/Private/HoudiniCSVAssetActions.cpp
+++ b/Engine/Plugins/FX/HoudiniNiagara/Source/HoudiniNiagaraEditor/Private/HoudiniCSVAssetActions.cpp
@@ -22,7 +22,7 @@
 */
 #include "HoudiniCSVAssetActions.h"
 
-#include "EditorMenuSubsystem.h"
+#include "ToolMenus.h"
 #include "HoudiniCSV.h"
 #include "EditorStyleSet.h"
 #include "Toolkits/AssetEditorToolkit.h"
@@ -43,7 +43,7 @@ bool FHoudiniCSVAssetActions::CanFilter()
 }
 
 
-void FHoudiniCSVAssetActions::GetActions(const TArray<UObject*>& InObjects, FEditorMenuSection& Section)
+void FHoudiniCSVAssetActions::GetActions(const TArray<UObject*>& InObjects, FToolMenuSection& Section)
 {
     FAssetTypeActions_Base::GetActions(InObjects, Section);
 
diff --git a/Engine/Plugins/FX/Niagara/Source/NiagaraEditor/Private/AssetTypeActions/AssetTypeActions_NiagaraEmitter.cpp b/Engine/Plugins/FX/Niagara/Source/NiagaraEditor/Private/AssetTypeActions/AssetTypeActions_NiagaraEmitter.cpp
--- a/Engine/Plugins/FX/Niagara/Source/NiagaraEditor/Private/AssetTypeActions/AssetTypeActions_NiagaraEmitter.cpp
+++ b/Engine/Plugins/FX/Niagara/Source/NiagaraEditor/Private/AssetTypeActions/AssetTypeActions_NiagaraEmitter.cpp
@@ -8,7 +8,7 @@
 #include "NiagaraSystemFactoryNew.h"
 #include "ViewModels/NiagaraSystemViewModel.h"
 #include "NiagaraSystemScriptViewModel.h"
-#include "EditorMenuSubsystem.h"
+#include "ToolMenus.h"
 #include "EditorStyleSet.h"
 #include "Toolkits/AssetEditorToolkit.h"
 #include "IContentBrowserSingleton.h"
@@ -52,7 +52,7 @@ UClass* FAssetTypeActions_NiagaraEmitter::GetSupportedClass() const
 	return UNiagaraEmitter::StaticClass();
 }
 
-void FAssetTypeActions_NiagaraEmitter::GetActions(const TArray<UObject*>& InObjects, FEditorMenuSection& Section)
+void FAssetTypeActions_NiagaraEmitter::GetActions(const TArray<UObject*>& InObjects, FToolMenuSection& Section)
 {
 	auto NiagaraEmitters = GetTypedWeakObjectPtrs<UNiagaraEmitter>(InObjects);
 
diff --git a/Engine/Plugins/FX/Niagara/Source/NiagaraEditor/Private/AssetTypeActions/AssetTypeActions_NiagaraParameterCollection.cpp b/Engine/Plugins/FX/Niagara/Source/NiagaraEditor/Private/AssetTypeActions/AssetTypeActions_NiagaraParameterCollection.cpp
--- a/Engine/Plugins/FX/Niagara/Source/NiagaraEditor/Private/AssetTypeActions/AssetTypeActions_NiagaraParameterCollection.cpp
+++ b/Engine/Plugins/FX/Niagara/Source/NiagaraEditor/Private/AssetTypeActions/AssetTypeActions_NiagaraParameterCollection.cpp
@@ -9,7 +9,7 @@
 #include "IContentBrowserSingleton.h"
 #include "ContentBrowserModule.h"
 
-#include "EditorMenuSubsystem.h"
+#include "ToolMenus.h"
 #include "EditorStyleSet.h"
 
 #define LOCTEXT_NAMESPACE "AssetTypeActions_NiagaraParameterCollection"
@@ -34,7 +34,7 @@ void FAssetTypeActions_NiagaraParameterCollection::OpenAssetEditor(const TArray<
 	}
 }
 
-void FAssetTypeActions_NiagaraParameterCollection::GetActions(const TArray<UObject*>& InObjects, FEditorMenuSection& Section)
+void FAssetTypeActions_NiagaraParameterCollection::GetActions(const TArray<UObject*>& InObjects, FToolMenuSection& Section)
 {
 	auto Collections = GetTypedWeakObjectPtrs<UNiagaraParameterCollection>(InObjects);
 
diff --git a/Engine/Plugins/FX/Niagara/Source/NiagaraEditor/Private/AssetTypeActions/AssetTypeActions_NiagaraScript.cpp b/Engine/Plugins/FX/Niagara/Source/NiagaraEditor/Private/AssetTypeActions/AssetTypeActions_NiagaraScript.cpp
--- a/Engine/Plugins/FX/Niagara/Source/NiagaraEditor/Private/AssetTypeActions/AssetTypeActions_NiagaraScript.cpp
+++ b/Engine/Plugins/FX/Niagara/Source/NiagaraEditor/Private/AssetTypeActions/AssetTypeActions_NiagaraScript.cpp
@@ -6,7 +6,7 @@
 #include "NiagaraEditorStyle.h"
 #include "AssetData.h"
 #include "NiagaraEditorUtilities.h"
-#include "EditorMenuSubsystem.h"
+#include "ToolMenus.h"
 
 #define LOCTEXT_NAMESPACE "NiagaraScriptAssetTypeActions"
 
@@ -74,7 +74,7 @@ bool FAssetTypeActions_NiagaraScript::HasActions(const TArray<UObject*>& InObjec
 	return true;
 }
 
-void FAssetTypeActions_NiagaraScript::GetActions(const TArray<UObject*>& InObjects, FEditorMenuSection& Section)
+void FAssetTypeActions_NiagaraScript::GetActions(const TArray<UObject*>& InObjects, FToolMenuSection& Section)
 {
 	Section.AddMenuEntry(
 		"MarkDependentCompilableAssetsDirty",
diff --git a/Engine/Plugins/Media/MediaPlayerEditor/Source/MediaPlayerEditor/Private/AssetTools/MediaTextureActions.cpp b/Engine/Plugins/Media/MediaPlayerEditor/Source/MediaPlayerEditor/Private/AssetTools/MediaTextureActions.cpp
--- a/Engine/Plugins/Media/MediaPlayerEditor/Source/MediaPlayerEditor/Private/AssetTools/MediaTextureActions.cpp
+++ b/Engine/Plugins/Media/MediaPlayerEditor/Source/MediaPlayerEditor/Private/AssetTools/MediaTextureActions.cpp
@@ -2,7 +2,7 @@
 
 #include "AssetTools/MediaTextureActions.h"
 #include "Engine/Texture.h"
-#include "EditorMenuSubsystem.h"
+#include "ToolMenus.h"
 #include "Misc/PackageName.h"
 #include "IContentBrowserSingleton.h"
 #include "ContentBrowserModule.h"
@@ -23,7 +23,7 @@ bool FMediaTextureActions::CanFilter()
 	return true;
 }
 
-void FMediaTextureActions::GetActions(const TArray<UObject*>& InObjects, FEditorMenuSection& Section)
+void FMediaTextureActions::GetActions(const TArray<UObject*>& InObjects, FToolMenuSection& Section)
 {
 	auto Textures = GetTypedWeakObjectPtrs<UTexture>(InObjects);
 
diff --git a/Engine/Plugins/Media/MediaPlayerEditor/Source/MediaPlayerEditor/Private/AssetTools/MediaTextureActions.h b/Engine/Plugins/Media/MediaPlayerEditor/Source/MediaPlayerEditor/Private/AssetTools/MediaTextureActions.h
--- a/Engine/Plugins/Media/MediaPlayerEditor/Source/MediaPlayerEditor/Private/AssetTools/MediaTextureActions.h
+++ b/Engine/Plugins/Media/MediaPlayerEditor/Source/MediaPlayerEditor/Private/AssetTools/MediaTextureActions.h
@@ -7,7 +7,7 @@
 #include "Toolkits/IToolkitHost.h"
 #include "AssetTypeActions_Base.h"
 
-struct FEditorMenuSection;
+struct FToolMenuSection;
 class UTexture;
 
 /**
@@ -21,7 +21,7 @@ public:
 	//~ FAssetTypeActions_Base overrides
 
 	virtual bool CanFilter() override;
-	virtual void GetActions(const TArray<UObject*>& InObjects, FEditorMenuSection& Section) override;
+	virtual void GetActions(const TArray<UObject*>& InObjects, FToolMenuSection& Section) override;
 	virtual uint32 GetCategories() override;
 	virtual FText GetName() const override;
 	virtual UClass* GetSupportedClass() const override;
diff --git a/Engine/Plugins/Media/MediaPlayerEditor/Source/MediaPlayerEditor/Private/MediaPlayerEditorModule.cpp b/Engine/Plugins/Media/MediaPlayerEditor/Source/MediaPlayerEditor/Private/MediaPlayerEditorModule.cpp
--- a/Engine/Plugins/Media/MediaPlayerEditor/Source/MediaPlayerEditor/Private/MediaPlayerEditorModule.cpp
+++ b/Engine/Plugins/Media/MediaPlayerEditor/Source/MediaPlayerEditor/Private/MediaPlayerEditorModule.cpp
@@ -39,7 +39,7 @@
 #include "Models/MediaPlayerEditorCommands.h"
 #include "Shared/MediaPlayerEditorStyle.h"
 #include "Visualizers/MediaSoundComponentVisualizer.h"
-#include "EditorMenuSubsystem.h"
+#include "ToolMenus.h"
 #include "Widgets/SMediaPlayerEditorMedia.h"
 
 #include "MediaPlayerEditorLog.h"
diff --git a/Engine/Plugins/Media/MediaPlayerEditor/Source/MediaPlayerEditor/Private/Shared/MediaPlayerEditorMediaContext.h b/Engine/Plugins/Media/MediaPlayerEditor/Source/MediaPlayerEditor/Private/Shared/MediaPlayerEditorMediaContext.h
--- a/Engine/Plugins/Media/MediaPlayerEditor/Source/MediaPlayerEditor/Private/Shared/MediaPlayerEditorMediaContext.h
+++ b/Engine/Plugins/Media/MediaPlayerEditor/Source/MediaPlayerEditor/Private/Shared/MediaPlayerEditorMediaContext.h
@@ -5,14 +5,14 @@
 #include "CoreMinimal.h"
 #include "UObject/ObjectMacros.h"
 #include "UObject/Object.h"
-#include "EditorMenuSubsystem.h"
+#include "ToolMenus.h"
 
 #include "MediaPlayerEditorMediaContext.generated.h"
 
 class SMediaPlayerEditorMedia;
 
 UCLASS(BlueprintType)
-class MEDIAPLAYEREDITOR_API UMediaPlayerEditorMediaContext : public UEditorMenuContextBase
+class MEDIAPLAYEREDITOR_API UMediaPlayerEditorMediaContext : public UToolMenuContextBase
 {
 	GENERATED_BODY()
 
diff --git a/Engine/Plugins/Media/MediaPlayerEditor/Source/MediaPlayerEditor/Private/Widgets/SMediaPlayerEditorMedia.cpp b/Engine/Plugins/Media/MediaPlayerEditor/Source/MediaPlayerEditor/Private/Widgets/SMediaPlayerEditorMedia.cpp
--- a/Engine/Plugins/Media/MediaPlayerEditor/Source/MediaPlayerEditor/Private/Widgets/SMediaPlayerEditorMedia.cpp
+++ b/Engine/Plugins/Media/MediaPlayerEditor/Source/MediaPlayerEditor/Private/Widgets/SMediaPlayerEditorMedia.cpp
@@ -20,7 +20,7 @@
 #include "Widgets/Layout/SBorder.h"
 #include "Widgets/Notifications/SNotificationList.h"
 #include "Widgets/Text/STextBlock.h"
-#include "EditorMenuSubsystem.h"
+#include "ToolMenus.h"
 #include "Shared/MediaPlayerEditorMediaContext.h"
 
 #define LOCTEXT_NAMESPACE "SMediaPlayerEditorMedia"
@@ -128,15 +128,15 @@ void SMediaPlayerEditorMedia::HandleAssetPickerAssetEnterPressed(const TArray<FA
 
 void SMediaPlayerEditorMedia::RegisterMenus()
 {
-	UEditorMenu* Menu = UEditorMenuSubsystem::Get()->RegisterMenu(AssetPickerAssetContextMenuName);
+	UToolMenu* Menu = UToolMenus::Get()->RegisterMenu(AssetPickerAssetContextMenuName);
 	{
-		FEditorMenuSection& Section = Menu->AddSection("MediaSection", LOCTEXT("MediaSection", "Media"));
-		Section.AddEntry(FEditorMenuEntry::InitMenuEntry(
+		FToolMenuSection& Section = Menu->AddSection("MediaSection", LOCTEXT("MediaSection", "Media"));
+		Section.AddEntry(FToolMenuEntry::InitMenuEntry(
 			"Edit",
 			LOCTEXT("EditMenuAction", "Edit..."),
 			LOCTEXT("EditMenuActionTooltip", "Opens the selected asset for edit."),
 			FSlateIcon(FEditorStyle::GetStyleSetName(), "ContentBrowser.AssetActions.Edit"),
-			FEditorMenuExecuteAction::CreateLambda([](const FEditorMenuContext& InContext)
+			FToolMenuExecuteAction::CreateLambda([](const FToolMenuContext& InContext)
 			{
 				if (UMediaPlayerEditorMediaContext* Context = InContext.Find<UMediaPlayerEditorMediaContext>())
 				{
@@ -149,16 +149,16 @@ void SMediaPlayerEditorMedia::RegisterMenus()
 		));
 
 
-		Section.AddDynamicEntry("Open", FNewEditorMenuSectionDelegate::CreateLambda([](FEditorMenuSection& InSection)
+		Section.AddDynamicEntry("Open", FNewToolMenuSectionDelegate::CreateLambda([](FToolMenuSection& InSection)
 		{
 			if (UMediaPlayerEditorMediaContext* Context = InSection.FindContext<UMediaPlayerEditorMediaContext>())
 			{
-				InSection.AddEntry(FEditorMenuEntry::InitMenuEntry(
+				InSection.AddEntry(FToolMenuEntry::InitMenuEntry(
 					"Open",
 					LOCTEXT("OpenMenuAction", "Open"),
 					LOCTEXT("OpenMenuActionTooltip", "Open this media asset in the player"),
 					FSlateIcon(Context->StyleSetName, "MediaPlayerEditor.NextMedia.Small"),
-					FEditorMenuExecuteAction::CreateLambda([=](const FEditorMenuContext& InContext)
+					FToolMenuExecuteAction::CreateLambda([=](const FToolMenuContext& InContext)
 					{
 						if (UMediaPlayerEditorMediaContext* Context = InContext.Find<UMediaPlayerEditorMediaContext>())
 						{
@@ -174,13 +174,13 @@ void SMediaPlayerEditorMedia::RegisterMenus()
 	}
 
 	{
-		FEditorMenuSection& Section = Menu->AddSection("AssetSection", LOCTEXT("AssetSection", "Asset"));
-		Section.AddEntry(FEditorMenuEntry::InitMenuEntry(
+		FToolMenuSection& Section = Menu->AddSection("AssetSection", LOCTEXT("AssetSection", "Asset"));
+		Section.AddEntry(FToolMenuEntry::InitMenuEntry(
 			"FindInCbMenuAction",
 			LOCTEXT("FindInCbMenuAction", "Find in Content Browser"),
 			LOCTEXT("FindInCbMenuActionTooltip", "Summons the Content Browser and navigates to the selected asset"),
 			FSlateIcon(FEditorStyle::GetStyleSetName(), "SystemWideCommands.FindInContentBrowser"),
-			FEditorMenuExecuteAction::CreateLambda([](const FEditorMenuContext& InContext)
+			FToolMenuExecuteAction::CreateLambda([](const FToolMenuContext& InContext)
 			{
 				if (UMediaPlayerEditorMediaContext* Context = InContext.Find<UMediaPlayerEditorMediaContext>())
 				{
@@ -190,7 +190,7 @@ void SMediaPlayerEditorMedia::RegisterMenus()
 			})
 		));
 
-		Section.AddDynamicEntry("OpenInFileManager", FNewEditorMenuSectionDelegate::CreateLambda([](FEditorMenuSection& InSection)
+		Section.AddDynamicEntry("OpenInFileManager", FNewToolMenuSectionDelegate::CreateLambda([](FToolMenuSection& InSection)
 		{
 			if (UMediaPlayerEditorMediaContext* Context = InSection.FindContext<UMediaPlayerEditorMediaContext>())
 			{
@@ -202,7 +202,7 @@ void SMediaPlayerEditorMedia::RegisterMenus()
 					}
 
 					TWeakObjectPtr<UFileMediaSource> FileMediaSourceWeak = FileMediaSource;
-					InSection.AddEntry(FEditorMenuEntry::InitMenuEntry(
+					InSection.AddEntry(FToolMenuEntry::InitMenuEntry(
 						"OpenInFileManager",
 						FText::Format(LOCTEXT("OpenInFileManager", "Show Media File in {FileManagerName}"), Args),
 						LOCTEXT("OpenInFileManagerTooltip", "Finds the media file that this asset points to on disk"),
@@ -241,8 +241,8 @@ TSharedPtr<SWidget> SMediaPlayerEditorMedia::HandleAssetPickerGetAssetContextMen
 	UMediaPlayerEditorMediaContext* ContextObject = NewObject<UMediaPlayerEditorMediaContext>();
 	ContextObject->InitContext(SelectedAsset, Style->GetStyleSetName());
 	ContextObject->MediaPlayerEditorMedia = SharedThis(this);
-	FEditorMenuContext MenuContext(ContextObject);
-	return UEditorMenuSubsystem::Get()->GenerateWidget(AssetPickerAssetContextMenuName, MenuContext);
+	FToolMenuContext MenuContext(ContextObject);
+	return UToolMenus::Get()->GenerateWidget(AssetPickerAssetContextMenuName, MenuContext);
 }
 
 
diff --git a/Engine/Plugins/Runtime/SoundMod/Source/SoundModImporter/Private/AssetTypeActions_SoundMod.cpp b/Engine/Plugins/Runtime/SoundMod/Source/SoundModImporter/Private/AssetTypeActions_SoundMod.cpp
--- a/Engine/Plugins/Runtime/SoundMod/Source/SoundModImporter/Private/AssetTypeActions_SoundMod.cpp
+++ b/Engine/Plugins/Runtime/SoundMod/Source/SoundModImporter/Private/AssetTypeActions_SoundMod.cpp
@@ -1,7 +1,7 @@
 // Copyright 1998-2019 Epic Games, Inc. All Rights Reserved.
 
 #include "AssetTypeActions_SoundMod.h"
-#include "EditorMenuSubsystem.h"
+#include "ToolMenus.h"
 #include "Editor.h"
 #include "Components/AudioComponent.h"
 #include "SoundMod.h"
@@ -13,7 +13,7 @@ UClass* FAssetTypeActions_SoundMod::GetSupportedClass() const
 	return USoundMod::StaticClass();
 }
 
-void FAssetTypeActions_SoundMod::GetActions(const TArray<UObject*>& InObjects, FEditorMenuSection& Section)
+void FAssetTypeActions_SoundMod::GetActions(const TArray<UObject*>& InObjects, FToolMenuSection& Section)
 {
 	auto Sounds = GetTypedWeakObjectPtrs<USoundMod>(InObjects);
 
diff --git a/Engine/Plugins/Runtime/SoundUtilities/Source/SoundUtilitiesEditor/Private/SoundUtilitiesEditorModule.cpp b/Engine/Plugins/Runtime/SoundUtilities/Source/SoundUtilitiesEditor/Private/SoundUtilitiesEditorModule.cpp
--- a/Engine/Plugins/Runtime/SoundUtilities/Source/SoundUtilitiesEditor/Private/SoundUtilitiesEditorModule.cpp
+++ b/Engine/Plugins/Runtime/SoundUtilities/Source/SoundUtilitiesEditor/Private/SoundUtilitiesEditorModule.cpp
@@ -8,7 +8,7 @@
 #include "AssetToolsModule.h"
 #include "AssetTypeActions_SoundSimple.h"
 #include "AssetTypeActions_Base.h"
-#include "EditorMenuSubsystem.h"
+#include "ToolMenus.h"
 #include "AudioEditorModule.h"
 #include "SoundWaveAssetActionExtender.h"
 
@@ -21,13 +21,13 @@ void FSoundUtilitiesEditorModule::StartupModule()
 	// Register asset actions
 	AssetTools.RegisterAssetTypeActions(MakeShareable(new FAssetTypeActions_SoundSimple));
 
-	UEditorMenuSubsystem::RegisterStartupCallback(FSimpleMulticastDelegate::FDelegate::CreateRaw(this, &FSoundUtilitiesEditorModule::RegisterMenus));
+	UToolMenus::RegisterStartupCallback(FSimpleMulticastDelegate::FDelegate::CreateRaw(this, &FSoundUtilitiesEditorModule::RegisterMenus));
 }
 
 void FSoundUtilitiesEditorModule::ShutdownModule()
 {
-	UEditorMenuSubsystem::UnRegisterStartupCallback(this);
-	UEditorMenuSubsystem::UnregisterOwner("SoundUtilities");
+	UToolMenus::UnRegisterStartupCallback(this);
+	UToolMenus::UnregisterOwner("SoundUtilities");
 }
 
 void FSoundUtilitiesEditorModule::RegisterMenus()
diff --git a/Engine/Plugins/Runtime/SoundUtilities/Source/SoundUtilitiesEditor/Private/SoundWaveAssetActionExtender.cpp b/Engine/Plugins/Runtime/SoundUtilities/Source/SoundUtilitiesEditor/Private/SoundWaveAssetActionExtender.cpp
--- a/Engine/Plugins/Runtime/SoundUtilities/Source/SoundUtilitiesEditor/Private/SoundWaveAssetActionExtender.cpp
+++ b/Engine/Plugins/Runtime/SoundUtilities/Source/SoundUtilitiesEditor/Private/SoundWaveAssetActionExtender.cpp
@@ -1,7 +1,7 @@
 // Copyright 1998-2019 Epic Games, Inc. All Rights Reserved.
 
 #include "SoundWaveAssetActionExtender.h"
-#include "EditorMenuSubsystem.h"
+#include "ToolMenus.h"
 #include "AssetTypeActions_Base.h"
 #include "IContentBrowserSingleton.h"
 #include "ContentBrowserModule.h"
@@ -16,27 +16,27 @@
 
 void FSoundWaveAssetActionExtender::RegisterMenus()
 {
-	if (!UEditorMenuSubsystem::IsRunningEditorUI())
+	if (!UToolMenus::IsToolMenuUIEnabled())
 	{
 		return;
 	}
 
-	FEditorMenuOwnerScoped MenuOwner("SoundUtilities");
-	UEditorMenu* Menu = UEditorMenuSubsystem::Get()->ExtendMenu("ContentBrowser.AssetContextMenu.SoundWave");
-	FEditorMenuSection& Section = Menu->FindOrAddSection("GetAssetActions");
+	FToolMenuOwnerScoped MenuOwner("SoundUtilities");
+	UToolMenu* Menu = UToolMenus::Get()->ExtendMenu("ContentBrowser.AssetContextMenu.SoundWave");
+	FToolMenuSection& Section = Menu->FindOrAddSection("GetAssetActions");
 	
-	Section.AddDynamicEntry("SoundWaveAsset", FNewEditorMenuSectionDelegate::CreateLambda([](FEditorMenuSection& InSection)
+	Section.AddDynamicEntry("SoundWaveAsset", FNewToolMenuSectionDelegate::CreateLambda([](FToolMenuSection& InSection)
 	{
 		const TAttribute<FText> Label = LOCTEXT("SoundWave_CreateSimpleSound", "Create Simple Sound");
 		const TAttribute<FText> ToolTip = LOCTEXT("SoundWave_CreateSimpleSoundTooltip", "Creates a simple sound asset using the selected sound waves.");
 		const FSlateIcon Icon = FSlateIcon(FEditorStyle::GetStyleSetName(), "ClassIcon.SoundSimple");
-		const FEditorMenuExecuteAction UIAction = FEditorMenuExecuteAction::CreateStatic(&FSoundWaveAssetActionExtender::ExecuteCreateSimpleSound);
+		const FToolMenuExecuteAction UIAction = FToolMenuExecuteAction::CreateStatic(&FSoundWaveAssetActionExtender::ExecuteCreateSimpleSound);
 
 		InSection.AddMenuEntry("SoundWave_CreateSimpleSound", Label, ToolTip, Icon, UIAction);
 	}));
 }
 
-void FSoundWaveAssetActionExtender::ExecuteCreateSimpleSound(const FEditorMenuContext& MenuContext)
+void FSoundWaveAssetActionExtender::ExecuteCreateSimpleSound(const FToolMenuContext& MenuContext)
 {
 	UContentBrowserAssetContextMenuContext* Context = MenuContext.Find<UContentBrowserAssetContextMenuContext>();
 	if (!Context || Context->SelectedObjects.Num() == 0)
diff --git a/Engine/Plugins/Runtime/SoundUtilities/Source/SoundUtilitiesEditor/Public/SoundWaveAssetActionExtender.h b/Engine/Plugins/Runtime/SoundUtilities/Source/SoundUtilitiesEditor/Public/SoundWaveAssetActionExtender.h
--- a/Engine/Plugins/Runtime/SoundUtilities/Source/SoundUtilitiesEditor/Public/SoundWaveAssetActionExtender.h
+++ b/Engine/Plugins/Runtime/SoundUtilities/Source/SoundUtilitiesEditor/Public/SoundWaveAssetActionExtender.h
@@ -6,7 +6,7 @@ class FSoundWaveAssetActionExtender
 {
 public:
 	static void RegisterMenus();
-	static void GetExtendedActions(const struct FEditorMenuContext& MenuContext);
-	static void ExecuteCreateSimpleSound(const struct FEditorMenuContext& MenuContext);
+	static void GetExtendedActions(const struct FToolMenuContext& MenuContext);
+	static void ExecuteCreateSimpleSound(const struct FToolMenuContext& MenuContext);
 };
 
diff --git a/Engine/Plugins/Runtime/TimeSynth/Source/TimeSynthEditor/Private/TimeSynthEditorModule.cpp b/Engine/Plugins/Runtime/TimeSynth/Source/TimeSynthEditor/Private/TimeSynthEditorModule.cpp
--- a/Engine/Plugins/Runtime/TimeSynth/Source/TimeSynthEditor/Private/TimeSynthEditorModule.cpp
+++ b/Engine/Plugins/Runtime/TimeSynth/Source/TimeSynthEditor/Private/TimeSynthEditorModule.cpp
@@ -11,7 +11,7 @@
 #include "TimeSynthClip.h"
 #include "TimeSynthSoundWaveAssetActionExtender.h"
 #include "TimeSynthVolumeGroup.h"
-#include "EditorMenuSubsystem.h"
+#include "ToolMenus.h"
 #include "AudioEditorModule.h"
 
 
@@ -23,13 +23,13 @@ void FTimeSynthEditorModule::StartupModule()
 	AssetTools.RegisterAssetTypeActions(MakeShareable(new FAssetTypeActions_TimeSynthClip));
 	AssetTools.RegisterAssetTypeActions(MakeShareable(new FAssetTypeActions_TimeSynthVolumeGroup));
 
-	UEditorMenuSubsystem::RegisterStartupCallback(FSimpleMulticastDelegate::FDelegate::CreateRaw(this, &FTimeSynthEditorModule::RegisterMenus));
+	UToolMenus::RegisterStartupCallback(FSimpleMulticastDelegate::FDelegate::CreateRaw(this, &FTimeSynthEditorModule::RegisterMenus));
 }
 
 void FTimeSynthEditorModule::ShutdownModule()
 {
-	UEditorMenuSubsystem::UnRegisterStartupCallback(this);
-	UEditorMenuSubsystem::UnregisterOwner("TimeSynth");
+	UToolMenus::UnRegisterStartupCallback(this);
+	UToolMenus::UnregisterOwner("TimeSynth");
 }
 
 void FTimeSynthEditorModule::RegisterMenus()
diff --git a/Engine/Plugins/Runtime/TimeSynth/Source/TimeSynthEditor/Private/TimeSynthSoundWaveAssetActionExtender.cpp b/Engine/Plugins/Runtime/TimeSynth/Source/TimeSynthEditor/Private/TimeSynthSoundWaveAssetActionExtender.cpp
--- a/Engine/Plugins/Runtime/TimeSynth/Source/TimeSynthEditor/Private/TimeSynthSoundWaveAssetActionExtender.cpp
+++ b/Engine/Plugins/Runtime/TimeSynth/Source/TimeSynthEditor/Private/TimeSynthSoundWaveAssetActionExtender.cpp
@@ -2,7 +2,7 @@
 
 #include "TimeSynthSoundWaveAssetActionExtender.h"
 
-#include "EditorMenuSubsystem.h"
+#include "ToolMenus.h"
 #include "AssetTypeActions_Base.h"
 #include "IContentBrowserSingleton.h"
 #include "ContentBrowserModule.h"
@@ -22,16 +22,16 @@ namespace
 
 void FTimeSynthSoundWaveAssetActionExtender::RegisterMenus()
 {
-	if (!UEditorMenuSubsystem::IsRunningEditorUI())
+	if (!UToolMenus::IsToolMenuUIEnabled())
 	{
 		return;
 	}
 
-	FEditorMenuOwnerScoped MenuOwner("TimeSynth");
-	UEditorMenu* Menu = UEditorMenuSubsystem::Get()->ExtendMenu("ContentBrowser.AssetContextMenu.SoundWave");
-	FEditorMenuSection& Section = Menu->FindOrAddSection("GetAssetActions");
+	FToolMenuOwnerScoped MenuOwner("TimeSynth");
+	UToolMenu* Menu = UToolMenus::Get()->ExtendMenu("ContentBrowser.AssetContextMenu.SoundWave");
+	FToolMenuSection& Section = Menu->FindOrAddSection("GetAssetActions");
 
-	Section.AddDynamicEntry("TimeSynthSoundWaveAsset", FNewEditorMenuSectionDelegate::CreateLambda([](FEditorMenuSection& InSection)
+	Section.AddDynamicEntry("TimeSynthSoundWaveAsset", FNewToolMenuSectionDelegate::CreateLambda([](FToolMenuSection& InSection)
 	{
 		UContentBrowserAssetContextMenuContext* Context = InSection.FindContext<UContentBrowserAssetContextMenuContext>();
 		if (!Context || Context->SelectedObjects.Num() == 0)
@@ -43,19 +43,19 @@ void FTimeSynthSoundWaveAssetActionExtender::RegisterMenus()
 		const TAttribute<FText> ToolTip = LOCTEXT("SoundWave_CreateTimeSynthClipToolTip", "Creates time synth clip per sound wave selected");
 		const FSlateIcon Icon = FSlateIcon(FEditorStyle::GetStyleSetName(), "ClassIcon.SoundCue");
 
-		const FEditorMenuExecuteAction UIExecuteAction = FEditorMenuExecuteAction::CreateStatic(&FTimeSynthSoundWaveAssetActionExtender::ExecuteCreateTimeSyncClip);
+		const FToolMenuExecuteAction UIExecuteAction = FToolMenuExecuteAction::CreateStatic(&FTimeSynthSoundWaveAssetActionExtender::ExecuteCreateTimeSyncClip);
 
 		InSection.AddMenuEntry("SoundWave_CreateTimeSynthClip", Label, ToolTip, Icon, UIExecuteAction);
 
 		const TAttribute<FText> LabelSet = LOCTEXT("SoundWave_CreateTimeSynthClipSet", "Create Time Synth Clip Set");
 		const TAttribute<FText> ToolTipSet = LOCTEXT("SoundWave_CreateTimeSynthClipSetToolTip", "Creates time synth clip adding all selected sound waves to single clip as set.");
-		const FEditorMenuExecuteAction UIExecuteActionSet = FEditorMenuExecuteAction::CreateStatic(&FTimeSynthSoundWaveAssetActionExtender::ExecuteCreateTimeSyncClipSet);
+		const FToolMenuExecuteAction UIExecuteActionSet = FToolMenuExecuteAction::CreateStatic(&FTimeSynthSoundWaveAssetActionExtender::ExecuteCreateTimeSyncClipSet);
 
 		InSection.AddMenuEntry("SoundWave_CreateTimeSynthClipSet", LabelSet, ToolTipSet, Icon, UIExecuteActionSet);
 	}));
 }
 
-void FTimeSynthSoundWaveAssetActionExtender::ExecuteCreateTimeSyncClip(const FEditorMenuContext& MenuContext)
+void FTimeSynthSoundWaveAssetActionExtender::ExecuteCreateTimeSyncClip(const FToolMenuContext& MenuContext)
 {
 	UContentBrowserAssetContextMenuContext* Context = MenuContext.Find<UContentBrowserAssetContextMenuContext>();
 	if (!Context || Context->SelectedObjects.Num() == 0)
@@ -82,7 +82,7 @@ void FTimeSynthSoundWaveAssetActionExtender::ExecuteCreateTimeSyncClip(const FEd
 	}
 }
 
-void FTimeSynthSoundWaveAssetActionExtender::ExecuteCreateTimeSyncClipSet(const FEditorMenuContext& MenuContext)
+void FTimeSynthSoundWaveAssetActionExtender::ExecuteCreateTimeSyncClipSet(const FToolMenuContext& MenuContext)
 {
 	UContentBrowserAssetContextMenuContext* Context = MenuContext.Find<UContentBrowserAssetContextMenuContext>();
 	if (!Context || Context->SelectedObjects.Num() == 0)
diff --git a/Engine/Plugins/Runtime/TimeSynth/Source/TimeSynthEditor/Public/TimeSynthSoundWaveAssetActionExtender.h b/Engine/Plugins/Runtime/TimeSynth/Source/TimeSynthEditor/Public/TimeSynthSoundWaveAssetActionExtender.h
--- a/Engine/Plugins/Runtime/TimeSynth/Source/TimeSynthEditor/Public/TimeSynthSoundWaveAssetActionExtender.h
+++ b/Engine/Plugins/Runtime/TimeSynth/Source/TimeSynthEditor/Public/TimeSynthSoundWaveAssetActionExtender.h
@@ -6,7 +6,7 @@ class FTimeSynthSoundWaveAssetActionExtender
 {
 public:
 	static void RegisterMenus();
-	static void ExecuteCreateTimeSyncClip(const struct FEditorMenuContext& MenuContext);
-	static void ExecuteCreateTimeSyncClipSet(const struct FEditorMenuContext& MenuContext);
+	static void ExecuteCreateTimeSyncClip(const struct FToolMenuContext& MenuContext);
+	static void ExecuteCreateTimeSyncClipSet(const struct FToolMenuContext& MenuContext);
 };
 
diff --git a/Engine/Plugins/VirtualProduction/Takes/Source/TakeRecorder/Private/TakeRecorderModule.cpp b/Engine/Plugins/VirtualProduction/Takes/Source/TakeRecorder/Private/TakeRecorderModule.cpp
--- a/Engine/Plugins/VirtualProduction/Takes/Source/TakeRecorder/Private/TakeRecorderModule.cpp
+++ b/Engine/Plugins/VirtualProduction/Takes/Source/TakeRecorder/Private/TakeRecorderModule.cpp
@@ -16,7 +16,7 @@
 #include "Framework/Docking/LayoutExtender.h"
 #include "Framework/MultiBox/MultiBoxExtender.h"
 #include "Framework/MultiBox/MultiBoxBuilder.h"
-#include "EditorMenuSubsystem.h"
+#include "ToolMenus.h"
 #include "ContentBrowserMenuContexts.h"
 #include "Features/IModularFeatures.h"
 #include "ITakeRecorderDropHandler.h"
@@ -168,21 +168,21 @@ namespace
 void FTakeRecorderModule::RegisterMenus()
 {
 #if WITH_EDITOR
-	if (!UEditorMenuSubsystem::IsRunningEditorUI())
+	if (!UToolMenus::IsToolMenuUIEnabled())
 	{
 		return;
 	}
 
-	FEditorMenuOwnerScoped MenuOwner("TakeRecorder");
-	UEditorMenuSubsystem* EditorMenus = UEditorMenuSubsystem::Get();
-	UEditorMenu* Menu = EditorMenus->ExtendMenu("ContentBrowser.AssetContextMenu.LevelSequence");
+	FToolMenuOwnerScoped MenuOwner("TakeRecorder");
+	UToolMenus* ToolMenus = UToolMenus::Get();
+	UToolMenu* Menu = ToolMenus->ExtendMenu("ContentBrowser.AssetContextMenu.LevelSequence");
 	if (!Menu)
 	{
 		return;
 	}
 
-	FEditorMenuSection& Section = Menu->FindOrAddSection("GetAssetActions");
-	Section.AddDynamicEntry("TakeRecorderActions", FNewEditorMenuSectionDelegate::CreateLambda([](FEditorMenuSection& InSection)
+	FToolMenuSection& Section = Menu->FindOrAddSection("GetAssetActions");
+	Section.AddDynamicEntry("TakeRecorderActions", FNewToolMenuSectionDelegate::CreateLambda([](FToolMenuSection& InSection)
 	{
 		UContentBrowserAssetContextMenuContext* Context = InSection.FindContext<UContentBrowserAssetContextMenuContext>();
 		if (!Context)
@@ -289,7 +289,7 @@ void FTakeRecorderModule::StartupModule()
 #if WITH_EDITOR
 	if (GIsEditor)
 	{
-		if (UEditorMenuSubsystem::TryGet())
+		if (UToolMenus::TryGet())
 		{
 			RegisterMenus();
 		}
diff --git a/Engine/Source/Developer/AssetTools/Private/AssetTools.cpp b/Engine/Source/Developer/AssetTools/Private/AssetTools.cpp
--- a/Engine/Source/Developer/AssetTools/Private/AssetTools.cpp
+++ b/Engine/Source/Developer/AssetTools/Private/AssetTools.cpp
@@ -26,7 +26,7 @@
 #include "AssetToolsLog.h"
 #include "AssetToolsModule.h"
 #include "Framework/MultiBox/MultiBoxBuilder.h"
-#include "EditorMenuSubsystem.h"
+#include "ToolMenus.h"
 #include "IClassTypeActions.h"
 #include "AssetTypeActions/AssetTypeActions_Blueprint.h"
 #include "AssetTypeActions/AssetTypeActions_Curve.h"
diff --git a/Engine/Source/Developer/AssetTools/Private/AssetTools.h b/Engine/Source/Developer/AssetTools/Private/AssetTools.h
--- a/Engine/Source/Developer/AssetTools/Private/AssetTools.h
+++ b/Engine/Source/Developer/AssetTools/Private/AssetTools.h
@@ -10,7 +10,7 @@
 #include "AssetTools.generated.h"
 
 class FAssetFixUpRedirectors;
-class UEditorMenu;
+class UToolMenu;
 class IClassTypeActions;
 class UAutomatedAssetImportData;
 class UFactory;
diff --git a/Engine/Source/Developer/AssetTools/Private/AssetTypeActions/AssetTypeActions_AnimBlueprint.cpp b/Engine/Source/Developer/AssetTools/Private/AssetTypeActions/AssetTypeActions_AnimBlueprint.cpp
--- a/Engine/Source/Developer/AssetTools/Private/AssetTypeActions/AssetTypeActions_AnimBlueprint.cpp
+++ b/Engine/Source/Developer/AssetTools/Private/AssetTypeActions/AssetTypeActions_AnimBlueprint.cpp
@@ -1,7 +1,7 @@
 // Copyright 1998-2019 Epic Games, Inc. All Rights Reserved.
 
 #include "AssetTypeActions/AssetTypeActions_AnimBlueprint.h"
-#include "EditorMenuSubsystem.h"
+#include "ToolMenus.h"
 #include "Widgets/Layout/SBorder.h"
 #include "Misc/MessageDialog.h"
 #include "Widgets/Images/SImage.h"
@@ -21,7 +21,7 @@
 
 #define LOCTEXT_NAMESPACE "AssetTypeActions"
 
-void FAssetTypeActions_AnimBlueprint::GetActions(const TArray<UObject*>& InObjects, FEditorMenuSection& Section)
+void FAssetTypeActions_AnimBlueprint::GetActions(const TArray<UObject*>& InObjects, FToolMenuSection& Section)
 {
 	FAssetTypeActions_Blueprint::GetActions(InObjects, Section);
 
@@ -38,18 +38,18 @@ void FAssetTypeActions_AnimBlueprint::GetActions(const TArray<UObject*>& InObjec
 			)
 		);
 
-	Section.AddEntry(FEditorMenuEntry::InitSubMenu(
+	Section.AddEntry(FToolMenuEntry::InitSubMenu(
 		NAME_None,
 		"RetargetBlueprintSubmenu",
 		LOCTEXT("RetargetBlueprintSubmenu", "Retarget Anim Blueprints"),
 		LOCTEXT("RetargetBlueprintSubmenu_ToolTip", "Opens the retarget blueprints menu"),
-		FNewEditorMenuDelegate::CreateSP( this, &FAssetTypeActions_AnimBlueprint::FillRetargetMenu, InObjects),
+		FNewToolMenuDelegate::CreateSP( this, &FAssetTypeActions_AnimBlueprint::FillRetargetMenu, InObjects),
 		false,
 		FSlateIcon(FEditorStyle::GetStyleSetName(), "Persona.AssetActions.RetargetSkeleton")
 		));
 }
 
-void FAssetTypeActions_AnimBlueprint::FillRetargetMenu(UEditorMenu* MenuBuilder, const TArray<UObject*> InObjects)
+void FAssetTypeActions_AnimBlueprint::FillRetargetMenu(UToolMenu* MenuBuilder, const TArray<UObject*> InObjects)
 {
 	bool bAllSkeletonsNull = true;
 
@@ -65,7 +65,7 @@ void FAssetTypeActions_AnimBlueprint::FillRetargetMenu(UEditorMenu* MenuBuilder,
 		}
 	}
 
-	FEditorMenuSection& Section = MenuBuilder->AddSection("Section");
+	FToolMenuSection& Section = MenuBuilder->AddSection("Section");
 	if(bAllSkeletonsNull)
 	{
 		Section.AddMenuEntry(
diff --git a/Engine/Source/Developer/AssetTools/Private/AssetTypeActions/AssetTypeActions_AnimBlueprint.h b/Engine/Source/Developer/AssetTools/Private/AssetTypeActions/AssetTypeActions_AnimBlueprint.h
--- a/Engine/Source/Developer/AssetTools/Private/AssetTypeActions/AssetTypeActions_AnimBlueprint.h
+++ b/Engine/Source/Developer/AssetTools/Private/AssetTypeActions/AssetTypeActions_AnimBlueprint.h
@@ -17,7 +17,7 @@ public:
 	virtual FText GetName() const override { return NSLOCTEXT("AssetTypeActions", "AssetTypeActions_AnimBlueprint", "Animation Blueprint"); }
 	virtual FColor GetTypeColor() const override { return FColor(200,116,0); }
 	virtual UClass* GetSupportedClass() const override { return UAnimBlueprint::StaticClass(); }
-	virtual void GetActions(const TArray<UObject*>& InObjects, struct FEditorMenuSection& Section) override;
+	virtual void GetActions(const TArray<UObject*>& InObjects, struct FToolMenuSection& Section) override;
 	virtual void OpenAssetEditor( const TArray<UObject*>& InObjects, TSharedPtr<class IToolkitHost> EditWithinLevelEditor = TSharedPtr<IToolkitHost>() ) override;
 	virtual uint32 GetCategories() override { return EAssetTypeCategories::Animation; }
 	virtual void PerformAssetDiff(UObject* Asset1, UObject* Asset2, const struct FRevisionInfo& OldRevision, const struct FRevisionInfo& NewRevision) const override;
@@ -30,7 +30,7 @@ public:
 private:
 
 	/** Handler to fill the retarget submenu */
-	void FillRetargetMenu(class UEditorMenu* Menu, const TArray<UObject*> InObjects);
+	void FillRetargetMenu(class UToolMenu* Menu, const TArray<UObject*> InObjects);
 
 	/** Handler for when FindSkeleton is selected */
 	void ExecuteFindSkeleton(TArray<TWeakObjectPtr<UAnimBlueprint>> Objects);
diff --git a/Engine/Source/Developer/AssetTools/Private/AssetTypeActions/AssetTypeActions_AnimCurveCompressionSettings.cpp b/Engine/Source/Developer/AssetTools/Private/AssetTypeActions/AssetTypeActions_AnimCurveCompressionSettings.cpp
--- a/Engine/Source/Developer/AssetTools/Private/AssetTypeActions/AssetTypeActions_AnimCurveCompressionSettings.cpp
+++ b/Engine/Source/Developer/AssetTools/Private/AssetTypeActions/AssetTypeActions_AnimCurveCompressionSettings.cpp
@@ -4,7 +4,7 @@
 #include "Animation/AnimSequence.h"
 #include "Dialogs/Dialogs.h"
 #include "EditorStyleSet.h"
-#include "EditorMenuSubsystem.h"
+#include "ToolMenus.h"
 #include "Misc/ScopedSlowTask.h"
 #include "UObject/UObjectIterator.h"
 
@@ -41,7 +41,7 @@ void FAssetTypeActions_AnimCurveCompressionSettings::AddToolbarExtension(FToolBa
 	Builder.EndSection();
 }
 
-void FAssetTypeActions_AnimCurveCompressionSettings::GetActions(const TArray<UObject*>& InObjects, FEditorMenuSection& Section)
+void FAssetTypeActions_AnimCurveCompressionSettings::GetActions(const TArray<UObject*>& InObjects, FToolMenuSection& Section)
 {
 	auto SettingAssets = GetTypedWeakObjectPtrs<UAnimCurveCompressionSettings>(InObjects);
 
diff --git a/Engine/Source/Developer/AssetTools/Private/AssetTypeActions/AssetTypeActions_AnimMontage.cpp b/Engine/Source/Developer/AssetTools/Private/AssetTypeActions/AssetTypeActions_AnimMontage.cpp
--- a/Engine/Source/Developer/AssetTools/Private/AssetTypeActions/AssetTypeActions_AnimMontage.cpp
+++ b/Engine/Source/Developer/AssetTools/Private/AssetTypeActions/AssetTypeActions_AnimMontage.cpp
@@ -3,14 +3,14 @@
 #include "AssetTypeActions/AssetTypeActions_AnimMontage.h"
 #include "Factories/AnimMontageFactory.h"
 #include "AnimationEditorUtils.h"
-#include "EditorMenuSubsystem.h"
+#include "ToolMenus.h"
 #include "AssetTools.h"
 #include "EditorStyleSet.h"
 
 #define LOCTEXT_NAMESPACE "AssetTypeActions"
 
 
-void FAssetTypeActions_AnimMontage::GetActions(const TArray<UObject*>& InObjects, FEditorMenuSection& Section)
+void FAssetTypeActions_AnimMontage::GetActions(const TArray<UObject*>& InObjects, FToolMenuSection& Section)
 {
 	auto Montages = GetTypedWeakObjectPtrs<UAnimMontage>(InObjects);
 
diff --git a/Engine/Source/Developer/AssetTools/Private/AssetTypeActions/AssetTypeActions_AnimSequence.cpp b/Engine/Source/Developer/AssetTools/Private/AssetTypeActions/AssetTypeActions_AnimSequence.cpp
--- a/Engine/Source/Developer/AssetTools/Private/AssetTypeActions/AssetTypeActions_AnimSequence.cpp
+++ b/Engine/Source/Developer/AssetTools/Private/AssetTypeActions/AssetTypeActions_AnimSequence.cpp
@@ -2,7 +2,7 @@
 
 #include "AssetTypeActions/AssetTypeActions_AnimSequence.h"
 #include "Animation/AnimSequence.h"
-#include "EditorMenuSubsystem.h"
+#include "ToolMenus.h"
 #include "EditorStyleSet.h"
 #include "EditorReimportHandler.h"
 #include "Animation/AnimMontage.h"
@@ -27,12 +27,12 @@ UClass* FAssetTypeActions_AnimSequence::GetSupportedClass() const
 	return UAnimSequence::StaticClass(); 
 }
 
-void FAssetTypeActions_AnimSequence::GetActions(const TArray<UObject*>& InObjects, FEditorMenuSection& Section)
+void FAssetTypeActions_AnimSequence::GetActions(const TArray<UObject*>& InObjects, FToolMenuSection& Section)
 {
 	auto Sequences = GetTypedWeakObjectPtrs<UAnimSequence>(InObjects);
 
 	// create menu
-	Section.AddEntry(FEditorMenuEntry::InitSubMenu(
+	Section.AddEntry(FToolMenuEntry::InitSubMenu(
 		NAME_None,
 		"CreateAnimSubmenu",
 		LOCTEXT("CreateAnimSubmenu", "Create"),
diff --git a/Engine/Source/Developer/AssetTools/Private/AssetTypeActions/AssetTypeActions_AnimStreamable.h b/Engine/Source/Developer/AssetTools/Private/AssetTypeActions/AssetTypeActions_AnimStreamable.h
--- a/Engine/Source/Developer/AssetTools/Private/AssetTypeActions/AssetTypeActions_AnimStreamable.h
+++ b/Engine/Source/Developer/AssetTools/Private/AssetTypeActions/AssetTypeActions_AnimStreamable.h
@@ -16,10 +16,10 @@ public:
 	virtual FColor GetTypeColor() const override { return FColor(181,230,29); }
 	virtual UClass* GetSupportedClass() const override { return UAnimStreamable::StaticClass(); }
 	virtual bool CanFilter() override { return true; }
-	virtual void GetActions(const TArray<UObject*>& InObjects, struct FEditorMenuSection& Section) override;
+	virtual void GetActions(const TArray<UObject*>& InObjects, struct FToolMenuSection& Section) override;
 };
 
-void FAssetTypeActions_AnimStreamable::GetActions(const TArray<UObject*>& InObjects, FEditorMenuSection& Section)
+void FAssetTypeActions_AnimStreamable::GetActions(const TArray<UObject*>& InObjects, FToolMenuSection& Section)
 {
 	/*auto Sequences = GetTypedWeakObjectPtrs<UAnimStreamable>(InObjects);
 
diff --git a/Engine/Source/Developer/AssetTools/Private/AssetTypeActions/AssetTypeActions_AnimationAsset.cpp b/Engine/Source/Developer/AssetTools/Private/AssetTypeActions/AssetTypeActions_AnimationAsset.cpp
--- a/Engine/Source/Developer/AssetTools/Private/AssetTypeActions/AssetTypeActions_AnimationAsset.cpp
+++ b/Engine/Source/Developer/AssetTools/Private/AssetTypeActions/AssetTypeActions_AnimationAsset.cpp
@@ -12,14 +12,14 @@
 #include "SSkeletonWidget.h"
 #include "IAnimationEditorModule.h"
 #include "Preferences/PersonaOptions.h"
-#include "EditorMenuSubsystem.h"
+#include "ToolMenus.h"
 #include "Framework/Application/SlateApplication.h"
 #include "Algo/Transform.h"
 #include "IPersonaToolkit.h"
 
 #define LOCTEXT_NAMESPACE "AssetTypeActions"
 
-void FAssetTypeActions_AnimationAsset::GetActions(const TArray<UObject*>& InObjects, FEditorMenuSection& Section)
+void FAssetTypeActions_AnimationAsset::GetActions(const TArray<UObject*>& InObjects, FToolMenuSection& Section)
 {
 	auto AnimAssets = GetTypedWeakObjectPtrs<UAnimationAsset>(InObjects);
 
@@ -45,7 +45,7 @@ void FAssetTypeActions_AnimationAsset::GetActions(const TArray<UObject*>& InObje
 			)
 		);
 
-	Section.AddEntry(FEditorMenuEntry::InitSubMenu(
+	Section.AddEntry(FToolMenuEntry::InitSubMenu(
 		NAME_None,
 		"RetargetAnimSubmenu", 
 		LOCTEXT("RetargetAnimSubmenu", "Retarget Anim Assets"),
diff --git a/Engine/Source/Developer/AssetTools/Private/AssetTypeActions/AssetTypeActions_Blueprint.cpp b/Engine/Source/Developer/AssetTools/Private/AssetTypeActions/AssetTypeActions_Blueprint.cpp
--- a/Engine/Source/Developer/AssetTools/Private/AssetTypeActions/AssetTypeActions_Blueprint.cpp
+++ b/Engine/Source/Developer/AssetTools/Private/AssetTypeActions/AssetTypeActions_Blueprint.cpp
@@ -1,7 +1,7 @@
 // Copyright 1998-2019 Epic Games, Inc. All Rights Reserved.
 
 #include "AssetTypeActions/AssetTypeActions_Blueprint.h"
-#include "EditorMenuSubsystem.h"
+#include "ToolMenus.h"
 #include "GameFramework/Actor.h"
 #include "Misc/PackageName.h"
 #include "Misc/MessageDialog.h"
@@ -19,7 +19,7 @@
 
 #define LOCTEXT_NAMESPACE "AssetTypeActions"
 
-void FAssetTypeActions_Blueprint::GetActions(const TArray<UObject*>& InObjects, FEditorMenuSection& Section)
+void FAssetTypeActions_Blueprint::GetActions(const TArray<UObject*>& InObjects, FToolMenuSection& Section)
 {
 	TArray<TWeakObjectPtr<UBlueprint>> Blueprints = GetTypedWeakObjectPtrs<UBlueprint>(InObjects);
 	
diff --git a/Engine/Source/Developer/AssetTools/Private/AssetTypeActions/AssetTypeActions_Class.cpp b/Engine/Source/Developer/AssetTools/Private/AssetTypeActions/AssetTypeActions_Class.cpp
--- a/Engine/Source/Developer/AssetTools/Private/AssetTypeActions/AssetTypeActions_Class.cpp
+++ b/Engine/Source/Developer/AssetTools/Private/AssetTypeActions/AssetTypeActions_Class.cpp
@@ -1,7 +1,7 @@
 // Copyright 1998-2019 Epic Games, Inc. All Rights Reserved.
 
 #include "AssetTypeActions/AssetTypeActions_Class.h"
-#include "EditorMenuSubsystem.h"
+#include "ToolMenus.h"
 #include "AssetData.h"
 #include "HAL/FileManager.h"
 #include "EditorStyleSet.h"
@@ -17,7 +17,7 @@ bool FAssetTypeActions_Class::HasActions(const TArray<UObject*>& InObjects) cons
 	return true;
 }
 
-void FAssetTypeActions_Class::GetActions(const TArray<UObject*>& InObjects, FEditorMenuSection& Section)
+void FAssetTypeActions_Class::GetActions(const TArray<UObject*>& InObjects, FToolMenuSection& Section)
 {
 	UClass *const BaseClass = (InObjects.Num() == 1) ? Cast<UClass>(InObjects[0]) : nullptr;
 
diff --git a/Engine/Source/Developer/AssetTools/Private/AssetTypeActions/AssetTypeActions_CurveLinearColorAtlas.cpp b/Engine/Source/Developer/AssetTools/Private/AssetTypeActions/AssetTypeActions_CurveLinearColorAtlas.cpp
--- a/Engine/Source/Developer/AssetTools/Private/AssetTypeActions/AssetTypeActions_CurveLinearColorAtlas.cpp
+++ b/Engine/Source/Developer/AssetTools/Private/AssetTypeActions/AssetTypeActions_CurveLinearColorAtlas.cpp
@@ -1,7 +1,7 @@
 // Copyright 1998-2019 Epic Games, Inc. All Rights Reserved.
 
 #include "AssetTypeActions/AssetTypeActions_CurveLinearColorAtlas.h"
-#include "EditorMenuSubsystem.h"
+#include "ToolMenus.h"
 #include "Engine/Texture2D.h"
 #include "EditorStyleSet.h"
 #include "Engine/TextureCube.h"
@@ -11,7 +11,7 @@
 
 #define LOCTEXT_NAMESPACE "AssetTypeActions"
 
-void FAssetTypeActions_CurveLinearColorAtlas::GetActions(const TArray<UObject*>& InObjects, FEditorMenuSection& Section)
+void FAssetTypeActions_CurveLinearColorAtlas::GetActions(const TArray<UObject*>& InObjects, FToolMenuSection& Section)
 {
 	FAssetTypeActions_Texture::GetActions(InObjects, Section);
 }
diff --git a/Engine/Source/Developer/AssetTools/Private/AssetTypeActions/AssetTypeActions_CurveTable.cpp b/Engine/Source/Developer/AssetTools/Private/AssetTypeActions/AssetTypeActions_CurveTable.cpp
--- a/Engine/Source/Developer/AssetTools/Private/AssetTypeActions/AssetTypeActions_CurveTable.cpp
+++ b/Engine/Source/Developer/AssetTools/Private/AssetTypeActions/AssetTypeActions_CurveTable.cpp
@@ -1,7 +1,7 @@
 // Copyright 1998-2019 Epic Games, Inc. All Rights Reserved.
 
 #include "AssetTypeActions/AssetTypeActions_CurveTable.h"
-#include "EditorMenuSubsystem.h"
+#include "ToolMenus.h"
 #include "Misc/FileHelper.h"
 #include "EditorFramework/AssetImportData.h"
 #include "Framework/Application/SlateApplication.h"
@@ -11,7 +11,7 @@
 
 #define LOCTEXT_NAMESPACE "AssetTypeActions"
 
-void FAssetTypeActions_CurveTable::GetActions(const TArray<UObject*>& InObjects, FEditorMenuSection& Section)
+void FAssetTypeActions_CurveTable::GetActions(const TArray<UObject*>& InObjects, FToolMenuSection& Section)
 {
 	auto Tables = GetTypedWeakObjectPtrs<UObject>(InObjects);
 
diff --git a/Engine/Source/Developer/AssetTools/Private/AssetTypeActions/AssetTypeActions_DataTable.cpp b/Engine/Source/Developer/AssetTools/Private/AssetTypeActions/AssetTypeActions_DataTable.cpp
--- a/Engine/Source/Developer/AssetTools/Private/AssetTypeActions/AssetTypeActions_DataTable.cpp
+++ b/Engine/Source/Developer/AssetTools/Private/AssetTypeActions/AssetTypeActions_DataTable.cpp
@@ -1,7 +1,7 @@
 // Copyright 1998-2019 Epic Games, Inc. All Rights Reserved.
 
 #include "AssetTypeActions/AssetTypeActions_DataTable.h"
-#include "EditorMenuSubsystem.h"
+#include "ToolMenus.h"
 #include "Misc/FileHelper.h"
 #include "EditorFramework/AssetImportData.h"
 #include "Dialogs/Dialogs.h"
@@ -13,7 +13,7 @@
 
 #define LOCTEXT_NAMESPACE "AssetTypeActions"
 
-void FAssetTypeActions_DataTable::GetActions(const TArray<UObject*>& InObjects, FEditorMenuSection& Section)
+void FAssetTypeActions_DataTable::GetActions(const TArray<UObject*>& InObjects, FToolMenuSection& Section)
 {
 	auto Tables = GetTypedWeakObjectPtrs<UObject>(InObjects);
 	
diff --git a/Engine/Source/Developer/AssetTools/Private/AssetTypeActions/AssetTypeActions_Font.cpp b/Engine/Source/Developer/AssetTools/Private/AssetTypeActions/AssetTypeActions_Font.cpp
--- a/Engine/Source/Developer/AssetTools/Private/AssetTypeActions/AssetTypeActions_Font.cpp
+++ b/Engine/Source/Developer/AssetTools/Private/AssetTypeActions/AssetTypeActions_Font.cpp
@@ -1,14 +1,14 @@
 // Copyright 1998-2019 Epic Games, Inc. All Rights Reserved.
 
 #include "AssetTypeActions/AssetTypeActions_Font.h"
-#include "EditorMenuSubsystem.h"
+#include "ToolMenus.h"
 #include "EditorStyleSet.h"
 #include "EditorReimportHandler.h"
 #include "FontEditorModule.h"
 
 #define LOCTEXT_NAMESPACE "AssetTypeActions"
 
-void FAssetTypeActions_Font::GetActions(const TArray<UObject*>& InObjects, FEditorMenuSection& Section)
+void FAssetTypeActions_Font::GetActions(const TArray<UObject*>& InObjects, FToolMenuSection& Section)
 {
 	auto Fonts = GetTypedWeakObjectPtrs<UFont>(InObjects);
 
diff --git a/Engine/Source/Developer/AssetTools/Private/AssetTypeActions/AssetTypeActions_FontFace.cpp b/Engine/Source/Developer/AssetTools/Private/AssetTypeActions/AssetTypeActions_FontFace.cpp
--- a/Engine/Source/Developer/AssetTools/Private/AssetTypeActions/AssetTypeActions_FontFace.cpp
+++ b/Engine/Source/Developer/AssetTools/Private/AssetTypeActions/AssetTypeActions_FontFace.cpp
@@ -2,13 +2,13 @@
 
 #include "AssetTypeActions/AssetTypeActions_FontFace.h"
 #include "FontEditorModule.h"
-#include "EditorMenuSubsystem.h"
+#include "ToolMenus.h"
 #include "EditorStyleSet.h"
 #include "EditorReimportHandler.h"
 
 #define LOCTEXT_NAMESPACE "AssetTypeActions"
 
-void FAssetTypeActions_FontFace::GetActions(const TArray<UObject*>& InObjects, FEditorMenuSection& Section)
+void FAssetTypeActions_FontFace::GetActions(const TArray<UObject*>& InObjects, FToolMenuSection& Section)
 {
 	auto FontFaces = GetTypedWeakObjectPtrs<UFontFace>(InObjects);
 
diff --git a/Engine/Source/Developer/AssetTools/Private/AssetTypeActions/AssetTypeActions_ForceFeedbackEffect.cpp b/Engine/Source/Developer/AssetTools/Private/AssetTypeActions/AssetTypeActions_ForceFeedbackEffect.cpp
--- a/Engine/Source/Developer/AssetTools/Private/AssetTypeActions/AssetTypeActions_ForceFeedbackEffect.cpp
+++ b/Engine/Source/Developer/AssetTools/Private/AssetTypeActions/AssetTypeActions_ForceFeedbackEffect.cpp
@@ -1,7 +1,7 @@
 // Copyright 1998-2019 Epic Games, Inc. All Rights Reserved.
 
 #include "AssetTypeActions/AssetTypeActions_ForceFeedbackEffect.h"
-#include "EditorMenuSubsystem.h"
+#include "ToolMenus.h"
 #include "EditorStyleSet.h"
 #include "GenericPlatform/IInputInterface.h"
 #include "AssetData.h"
@@ -17,7 +17,7 @@ UClass* FAssetTypeActions_ForceFeedbackEffect::GetSupportedClass() const
 	return UForceFeedbackEffect::StaticClass();
 }
 
-void FAssetTypeActions_ForceFeedbackEffect::GetActions(const TArray<UObject*>& InObjects, FEditorMenuSection& Section)
+void FAssetTypeActions_ForceFeedbackEffect::GetActions(const TArray<UObject*>& InObjects, FToolMenuSection& Section)
 {
 	TArray<TWeakObjectPtr<UForceFeedbackEffect>> Effects = GetTypedWeakObjectPtrs<UForceFeedbackEffect>(InObjects);
 
diff --git a/Engine/Source/Developer/AssetTools/Private/AssetTypeActions/AssetTypeActions_MaterialFunction.cpp b/Engine/Source/Developer/AssetTools/Private/AssetTypeActions/AssetTypeActions_MaterialFunction.cpp
--- a/Engine/Source/Developer/AssetTools/Private/AssetTypeActions/AssetTypeActions_MaterialFunction.cpp
+++ b/Engine/Source/Developer/AssetTools/Private/AssetTypeActions/AssetTypeActions_MaterialFunction.cpp
@@ -2,7 +2,7 @@
 
 #include "AssetTypeActions/AssetTypeActions_MaterialFunction.h"
 #include "Factories/MaterialFunctionInstanceFactory.h"
-#include "EditorMenuSubsystem.h"
+#include "ToolMenus.h"
 #include "UObject/UObjectHash.h"
 #include "UObject/UObjectIterator.h"
 #include "EditorStyleSet.h"
@@ -24,7 +24,7 @@ const FString LayerCompareString = (TEXT("MaterialLayer"));
 const FString BlendCompareString = (TEXT("MaterialLayerBlend"));
 
 
-void FAssetTypeActions_MaterialFunction::GetActions(const TArray<UObject*>& InObjects, FEditorMenuSection& Section)
+void FAssetTypeActions_MaterialFunction::GetActions(const TArray<UObject*>& InObjects, FToolMenuSection& Section)
 {
 	auto Functions = GetTypedWeakObjectPtrs<UMaterialFunctionInterface>(InObjects);
 
diff --git a/Engine/Source/Developer/AssetTools/Private/AssetTypeActions/AssetTypeActions_MaterialFunctionInstance.cpp b/Engine/Source/Developer/AssetTools/Private/AssetTypeActions/AssetTypeActions_MaterialFunctionInstance.cpp
--- a/Engine/Source/Developer/AssetTools/Private/AssetTypeActions/AssetTypeActions_MaterialFunctionInstance.cpp
+++ b/Engine/Source/Developer/AssetTools/Private/AssetTypeActions/AssetTypeActions_MaterialFunctionInstance.cpp
@@ -1,7 +1,7 @@
 // Copyright 1998-2019 Epic Games, Inc. All Rights Reserved.
 
 #include "AssetTypeActions/AssetTypeActions_MaterialFunctionInstance.h"
-#include "EditorMenuSubsystem.h"
+#include "ToolMenus.h"
 #include "UObject/UObjectHash.h"
 #include "UObject/UObjectIterator.h"
 #include "EditorStyleSet.h"
@@ -14,7 +14,7 @@
 
 #define LOCTEXT_NAMESPACE "AssetTypeActions"
 
-void FAssetTypeActions_MaterialFunctionInstance::GetActions(const TArray<UObject*>& InObjects, FEditorMenuSection& Section)
+void FAssetTypeActions_MaterialFunctionInstance::GetActions(const TArray<UObject*>& InObjects, FToolMenuSection& Section)
 {
 	auto MFIs = GetTypedWeakObjectPtrs<UMaterialFunctionInstance>(InObjects);
 
diff --git a/Engine/Source/Developer/AssetTools/Private/AssetTypeActions/AssetTypeActions_MaterialInstanceConstant.cpp b/Engine/Source/Developer/AssetTools/Private/AssetTypeActions/AssetTypeActions_MaterialInstanceConstant.cpp
--- a/Engine/Source/Developer/AssetTools/Private/AssetTypeActions/AssetTypeActions_MaterialInstanceConstant.cpp
+++ b/Engine/Source/Developer/AssetTools/Private/AssetTypeActions/AssetTypeActions_MaterialInstanceConstant.cpp
@@ -1,14 +1,14 @@
 // Copyright 1998-2019 Epic Games, Inc. All Rights Reserved.
 
 #include "AssetTypeActions/AssetTypeActions_MaterialInstanceConstant.h"
-#include "EditorMenuSubsystem.h"
+#include "ToolMenus.h"
 #include "EditorStyleSet.h"
 #include "AssetTools.h"
 #include "MaterialEditorModule.h"
 
 #define LOCTEXT_NAMESPACE "AssetTypeActions"
 
-void FAssetTypeActions_MaterialInstanceConstant::GetActions(const TArray<UObject*>& InObjects, FEditorMenuSection& Section)
+void FAssetTypeActions_MaterialInstanceConstant::GetActions(const TArray<UObject*>& InObjects, FToolMenuSection& Section)
 {
 	auto MICs = GetTypedWeakObjectPtrs<UMaterialInstanceConstant>(InObjects);
 
diff --git a/Engine/Source/Developer/AssetTools/Private/AssetTypeActions/AssetTypeActions_MaterialInterface.cpp b/Engine/Source/Developer/AssetTools/Private/AssetTypeActions/AssetTypeActions_MaterialInterface.cpp
--- a/Engine/Source/Developer/AssetTools/Private/AssetTypeActions/AssetTypeActions_MaterialInterface.cpp
+++ b/Engine/Source/Developer/AssetTools/Private/AssetTypeActions/AssetTypeActions_MaterialInterface.cpp
@@ -2,7 +2,7 @@
 
 #include "AssetTypeActions/AssetTypeActions_MaterialInterface.h"
 #include "Misc/PackageName.h"
-#include "EditorMenuSubsystem.h"
+#include "ToolMenus.h"
 #include "EditorStyleSet.h"
 #include "Factories/MaterialInstanceConstantFactoryNew.h"
 #include "ThumbnailRendering/SceneThumbnailInfoWithPrimitive.h"
@@ -13,7 +13,7 @@
 
 #define LOCTEXT_NAMESPACE "AssetTypeActions"
 
-void FAssetTypeActions_MaterialInterface::GetActions(const TArray<UObject*>& InObjects, FEditorMenuSection& Section)
+void FAssetTypeActions_MaterialInterface::GetActions(const TArray<UObject*>& InObjects, FToolMenuSection& Section)
 {
 	auto MaterialInterfaces = GetTypedWeakObjectPtrs<UMaterialInterface>(InObjects);
 
diff --git a/Engine/Source/Developer/AssetTools/Private/AssetTypeActions/AssetTypeActions_ParticleSystem.cpp b/Engine/Source/Developer/AssetTools/Private/AssetTypeActions/AssetTypeActions_ParticleSystem.cpp
--- a/Engine/Source/Developer/AssetTools/Private/AssetTypeActions/AssetTypeActions_ParticleSystem.cpp
+++ b/Engine/Source/Developer/AssetTools/Private/AssetTypeActions/AssetTypeActions_ParticleSystem.cpp
@@ -1,7 +1,7 @@
 // Copyright 1998-2019 Epic Games, Inc. All Rights Reserved.
 
 #include "AssetTypeActions/AssetTypeActions_ParticleSystem.h"
-#include "EditorMenuSubsystem.h"
+#include "ToolMenus.h"
 #include "Misc/FeedbackContext.h"
 #include "Editor/Cascade/Public/CascadeModule.h"
 #include "Particles/ParticleEmitter.h"
@@ -9,7 +9,7 @@
 
 #define LOCTEXT_NAMESPACE "AssetTypeActions"
 
-void FAssetTypeActions_ParticleSystem::GetActions(const TArray<UObject*>& InObjects, FEditorMenuSection& Section)
+void FAssetTypeActions_ParticleSystem::GetActions(const TArray<UObject*>& InObjects, FToolMenuSection& Section)
 {
 	auto ParticleSystems = GetTypedWeakObjectPtrs<UParticleSystem>(InObjects);
 
diff --git a/Engine/Source/Developer/AssetTools/Private/AssetTypeActions/AssetTypeActions_Redirector.cpp b/Engine/Source/Developer/AssetTools/Private/AssetTypeActions/AssetTypeActions_Redirector.cpp
--- a/Engine/Source/Developer/AssetTools/Private/AssetTypeActions/AssetTypeActions_Redirector.cpp
+++ b/Engine/Source/Developer/AssetTools/Private/AssetTypeActions/AssetTypeActions_Redirector.cpp
@@ -1,12 +1,12 @@
 // Copyright 1998-2019 Epic Games, Inc. All Rights Reserved.
 
 #include "AssetTypeActions/AssetTypeActions_Redirector.h"
-#include "EditorMenuSubsystem.h"
+#include "ToolMenus.h"
 #include "AssetTools.h"
 
 #define LOCTEXT_NAMESPACE "AssetTypeActions"
 
-void FAssetTypeActions_Redirector::GetActions(const TArray<UObject*>& InObjects, FEditorMenuSection& Section)
+void FAssetTypeActions_Redirector::GetActions(const TArray<UObject*>& InObjects, FToolMenuSection& Section)
 {
 	auto Redirectors = GetTypedWeakObjectPtrs<UObjectRedirector>(InObjects);
 
diff --git a/Engine/Source/Developer/AssetTools/Private/AssetTypeActions/AssetTypeActions_SkeletalMesh.cpp b/Engine/Source/Developer/AssetTools/Private/AssetTypeActions/AssetTypeActions_SkeletalMesh.cpp
--- a/Engine/Source/Developer/AssetTools/Private/AssetTypeActions/AssetTypeActions_SkeletalMesh.cpp
+++ b/Engine/Source/Developer/AssetTools/Private/AssetTypeActions/AssetTypeActions_SkeletalMesh.cpp
@@ -3,7 +3,7 @@
 #include "AssetTypeActions/AssetTypeActions_SkeletalMesh.h"
 #include "Animation/Skeleton.h"
 #include "Widgets/Text/STextBlock.h"
-#include "EditorMenuSubsystem.h"
+#include "ToolMenus.h"
 #include "Misc/PackageName.h"
 #include "Widgets/Layout/SBorder.h"
 #include "Misc/MessageDialog.h"
@@ -448,11 +448,11 @@ FAssetTypeActions_SkeletalMesh::FAssetTypeActions_SkeletalMesh()
 	AssetRegistryModule.Get().OnAssetRemoved().AddRaw(this, &FAssetTypeActions_SkeletalMesh::OnAssetRemoved);
 }
 
-void FAssetTypeActions_SkeletalMesh::GetActions(const TArray<UObject*>& InObjects, FEditorMenuSection& Section)
+void FAssetTypeActions_SkeletalMesh::GetActions(const TArray<UObject*>& InObjects, FToolMenuSection& Section)
 {
 	auto Meshes = GetTypedWeakObjectPtrs<USkeletalMesh>(InObjects);
 
-	Section.AddEntry(FEditorMenuEntry::InitSubMenu(
+	Section.AddEntry(FToolMenuEntry::InitSubMenu(
 			NAME_None,
 			"CreateSkeletalMeshSubmenu",
 			LOCTEXT("CreateSkeletalMeshSubmenu", "Create"),
@@ -462,7 +462,7 @@ void FAssetTypeActions_SkeletalMesh::GetActions(const TArray<UObject*>& InObject
 			FSlateIcon(FEditorStyle::GetStyleSetName(), "Persona.AssetActions.CreateAnimAsset")
 			));
 
-	Section.AddEntry(FEditorMenuEntry::InitSubMenu(
+	Section.AddEntry(FToolMenuEntry::InitSubMenu(
 		NAME_None,
 		"SkeletalMesh_LODImport",	
 		LOCTEXT("SkeletalMesh_LODImport", "Import LOD"),
@@ -478,7 +478,7 @@ void FAssetTypeActions_SkeletalMesh::GetActions(const TArray<UObject*>& InObject
 		FUIAction(FExecuteAction::CreateSP(this, &FAssetTypeActions_SkeletalMesh::ExecuteImportClothing, Meshes)));
 
 	// skeleton menu
-	Section.AddEntry(FEditorMenuEntry::InitSubMenu(
+	Section.AddEntry(FToolMenuEntry::InitSubMenu(
 		NAME_None,
 		"SkeletonSubmenu",
 		LOCTEXT("SkeletonSubmenu", "Skeleton"),
diff --git a/Engine/Source/Developer/AssetTools/Private/AssetTypeActions/AssetTypeActions_Skeleton.cpp b/Engine/Source/Developer/AssetTools/Private/AssetTypeActions/AssetTypeActions_Skeleton.cpp
--- a/Engine/Source/Developer/AssetTools/Private/AssetTypeActions/AssetTypeActions_Skeleton.cpp
+++ b/Engine/Source/Developer/AssetTools/Private/AssetTypeActions/AssetTypeActions_Skeleton.cpp
@@ -2,7 +2,7 @@
 
 #include "AssetTypeActions/AssetTypeActions_Skeleton.h"
 #include "Widgets/Text/STextBlock.h"
-#include "EditorMenuSubsystem.h"
+#include "ToolMenus.h"
 #include "Engine/SkeletalMesh.h"
 #include "Animation/AnimationAsset.h"
 #include "Animation/AnimSequenceBase.h"
@@ -425,12 +425,12 @@ FCreateRigDlg::EResult FCreateRigDlg::ShowModal()
 }
 
 ///////////////////////////////
-void FAssetTypeActions_Skeleton::GetActions(const TArray<UObject*>& InObjects, FEditorMenuSection& Section)
+void FAssetTypeActions_Skeleton::GetActions(const TArray<UObject*>& InObjects, FToolMenuSection& Section)
 {
 	auto Skeletons = GetTypedWeakObjectPtrs<USkeleton>(InObjects);
 
 	// create menu
-	Section.AddEntry(FEditorMenuEntry::InitSubMenu(
+	Section.AddEntry(FToolMenuEntry::InitSubMenu(
 			NAME_None,
 			"CreateSkeletonSubmenu",
 			LOCTEXT("CreateSkeletonSubmenu", "Create"),
diff --git a/Engine/Source/Developer/AssetTools/Private/AssetTypeActions/AssetTypeActions_SlateWidgetStyle.cpp b/Engine/Source/Developer/AssetTools/Private/AssetTypeActions/AssetTypeActions_SlateWidgetStyle.cpp
--- a/Engine/Source/Developer/AssetTools/Private/AssetTypeActions/AssetTypeActions_SlateWidgetStyle.cpp
+++ b/Engine/Source/Developer/AssetTools/Private/AssetTypeActions/AssetTypeActions_SlateWidgetStyle.cpp
@@ -1,13 +1,13 @@
 // Copyright 1998-2019 Epic Games, Inc. All Rights Reserved.
 
 #include "AssetTypeActions/AssetTypeActions_SlateWidgetStyle.h"
-#include "EditorMenuSubsystem.h"
+#include "ToolMenus.h"
 
 
 #define LOCTEXT_NAMESPACE "AssetTypeActions"
 
 
-void FAssetTypeActions_SlateWidgetStyle::GetActions(const TArray<UObject*>& InObjects, FEditorMenuSection& Section)
+void FAssetTypeActions_SlateWidgetStyle::GetActions(const TArray<UObject*>& InObjects, FToolMenuSection& Section)
 {
 	auto Styles = GetTypedWeakObjectPtrs< USlateWidgetStyleAsset >(InObjects);
 }
diff --git a/Engine/Source/Developer/AssetTools/Private/AssetTypeActions/AssetTypeActions_StaticMesh.cpp b/Engine/Source/Developer/AssetTools/Private/AssetTypeActions/AssetTypeActions_StaticMesh.cpp
--- a/Engine/Source/Developer/AssetTools/Private/AssetTypeActions/AssetTypeActions_StaticMesh.cpp
+++ b/Engine/Source/Developer/AssetTools/Private/AssetTypeActions/AssetTypeActions_StaticMesh.cpp
@@ -2,7 +2,7 @@
 
 #include "AssetTypeActions/AssetTypeActions_StaticMesh.h"
 #include "Framework/MultiBox/MultiBoxBuilder.h"
-#include "EditorMenuSubsystem.h"
+#include "ToolMenus.h"
 #include "EditorStyleSet.h"
 #include "EditorFramework/AssetImportData.h"
 #include "ThumbnailRendering/SceneThumbnailInfo.h"
@@ -24,7 +24,7 @@ static TAutoConsoleVariable<int32> CVarEnableSaveGeneratedLODsInPackage(
 	TEXT("1 - Enable this option and save the LODs in the Package.\n"),
 	ECVF_Default);
 
-void FAssetTypeActions_StaticMesh::GetActions(const TArray<UObject*>& InObjects, FEditorMenuSection& Section)
+void FAssetTypeActions_StaticMesh::GetActions(const TArray<UObject*>& InObjects, FToolMenuSection& Section)
 {
 	auto Meshes = GetTypedWeakObjectPtrs<UStaticMesh>(InObjects);
 
@@ -42,7 +42,7 @@ void FAssetTypeActions_StaticMesh::GetActions(const TArray<UObject*>& InObjects,
 			);
 	}
 
-	Section.AddEntry(FEditorMenuEntry::InitSubMenu(
+	Section.AddEntry(FToolMenuEntry::InitSubMenu(
 		NAME_None,
 		"StaticMesh_LODMenu",
 		NSLOCTEXT("AssetTypeActions_StaticMesh", "StaticMesh_LODMenu", "Level Of Detail"),
diff --git a/Engine/Source/Developer/AssetTools/Private/AssetTypeActions/AssetTypeActions_StaticMesh.h b/Engine/Source/Developer/AssetTools/Private/AssetTypeActions/AssetTypeActions_StaticMesh.h
--- a/Engine/Source/Developer/AssetTools/Private/AssetTypeActions/AssetTypeActions_StaticMesh.h
+++ b/Engine/Source/Developer/AssetTools/Private/AssetTypeActions/AssetTypeActions_StaticMesh.h
@@ -7,7 +7,7 @@
 #include "Toolkits/IToolkitHost.h"
 #include "AssetTypeActions_Base.h"
 
-struct FEditorMenuSection;
+struct FToolMenuSection;
 class FMenuBuilder;
 
 class FAssetTypeActions_StaticMesh : public FAssetTypeActions_Base
@@ -18,7 +18,7 @@ public:
 	virtual FColor GetTypeColor() const override { return FColor(0, 255, 255); }
 	virtual UClass* GetSupportedClass() const override { return UStaticMesh::StaticClass(); }
 	virtual bool HasActions( const TArray<UObject*>& InObjects ) const override { return true; }
-	virtual void GetActions(const TArray<UObject*>& InObjects, FEditorMenuSection& Section) override;
+	virtual void GetActions(const TArray<UObject*>& InObjects, FToolMenuSection& Section) override;
 	virtual void OpenAssetEditor( const TArray<UObject*>& InObjects, TSharedPtr<class IToolkitHost> EditWithinLevelEditor = TSharedPtr<IToolkitHost>() ) override;
 	virtual uint32 GetCategories() override { return EAssetTypeCategories::Basic; }
 	virtual class UThumbnailInfo* GetThumbnailInfo(UObject* Asset) const override;
diff --git a/Engine/Source/Developer/AssetTools/Private/AssetTypeActions/AssetTypeActions_Texture.cpp b/Engine/Source/Developer/AssetTools/Private/AssetTypeActions/AssetTypeActions_Texture.cpp
--- a/Engine/Source/Developer/AssetTools/Private/AssetTypeActions/AssetTypeActions_Texture.cpp
+++ b/Engine/Source/Developer/AssetTools/Private/AssetTypeActions/AssetTypeActions_Texture.cpp
@@ -21,7 +21,7 @@
 #include "Editor.h"
 #include "Widgets/Text/STextBlock.h"
 #include "Widgets/Images/SImage.h"
-#include "EditorMenuSubsystem.h"
+#include "ToolMenus.h"
 #include "Widgets/Layout/SBorder.h"
 #include "Misc/MessageDialog.h"
 #include "Misc/FeedbackContext.h"
@@ -44,7 +44,7 @@
 
 #define LOCTEXT_NAMESPACE "AssetTypeActions"
 
-void FAssetTypeActions_Texture::GetActions(const TArray<UObject*>& InObjects, FEditorMenuSection& Section)
+void FAssetTypeActions_Texture::GetActions(const TArray<UObject*>& InObjects, FToolMenuSection& Section)
 {
 	auto Textures = GetTypedWeakObjectPtrs<UTexture>(InObjects);
 	bool bHasVirtualTextures = false;
diff --git a/Engine/Source/Developer/AssetTools/Private/AssetTypeActions/AssetTypeActions_Texture2D.cpp b/Engine/Source/Developer/AssetTools/Private/AssetTypeActions/AssetTypeActions_Texture2D.cpp
--- a/Engine/Source/Developer/AssetTools/Private/AssetTypeActions/AssetTypeActions_Texture2D.cpp
+++ b/Engine/Source/Developer/AssetTools/Private/AssetTypeActions/AssetTypeActions_Texture2D.cpp
@@ -1,7 +1,7 @@
 // Copyright 1998-2019 Epic Games, Inc. All Rights Reserved.
 
 #include "AssetTypeActions/AssetTypeActions_Texture2D.h"
-#include "EditorMenuSubsystem.h"
+#include "ToolMenus.h"
 #include "Misc/PackageName.h"
 #include "EditorStyleSet.h"
 #include "Factories/SlateBrushAssetFactory.h"
@@ -14,7 +14,7 @@
 
 #define LOCTEXT_NAMESPACE "AssetTypeActions"
 
-void FAssetTypeActions_Texture2D::GetActions(const TArray<UObject*>& InObjects, FEditorMenuSection& Section)
+void FAssetTypeActions_Texture2D::GetActions(const TArray<UObject*>& InObjects, FToolMenuSection& Section)
 {
 	FAssetTypeActions_Texture::GetActions(InObjects, Section);
 
diff --git a/Engine/Source/Developer/AssetTools/Private/AssetTypeActions/AssetTypeActions_TextureRenderTarget.cpp b/Engine/Source/Developer/AssetTools/Private/AssetTypeActions/AssetTypeActions_TextureRenderTarget.cpp
--- a/Engine/Source/Developer/AssetTools/Private/AssetTypeActions/AssetTypeActions_TextureRenderTarget.cpp
+++ b/Engine/Source/Developer/AssetTools/Private/AssetTypeActions/AssetTypeActions_TextureRenderTarget.cpp
@@ -1,7 +1,7 @@
 // Copyright 1998-2019 Epic Games, Inc. All Rights Reserved.
 
 #include "AssetTypeActions/AssetTypeActions_TextureRenderTarget.h"
-#include "EditorMenuSubsystem.h"
+#include "ToolMenus.h"
 #include "Engine/Texture2D.h"
 #include "EditorStyleSet.h"
 #include "Engine/TextureCube.h"
@@ -11,7 +11,7 @@
 
 #define LOCTEXT_NAMESPACE "AssetTypeActions"
 
-void FAssetTypeActions_TextureRenderTarget::GetActions(const TArray<UObject*>& InObjects, FEditorMenuSection& Section)
+void FAssetTypeActions_TextureRenderTarget::GetActions(const TArray<UObject*>& InObjects, FToolMenuSection& Section)
 {
 	FAssetTypeActions_Texture::GetActions(InObjects, Section);
 
diff --git a/Engine/Source/Developer/AssetTools/Public/AssetTypeActions/AssetTypeActions_Blueprint.h b/Engine/Source/Developer/AssetTools/Public/AssetTypeActions/AssetTypeActions_Blueprint.h
--- a/Engine/Source/Developer/AssetTools/Public/AssetTypeActions/AssetTypeActions_Blueprint.h
+++ b/Engine/Source/Developer/AssetTools/Public/AssetTypeActions/AssetTypeActions_Blueprint.h
@@ -19,7 +19,7 @@ public:
 	virtual FColor GetTypeColor() const override { return FColor( 63, 126, 255 ); }
 	virtual UClass* GetSupportedClass() const override { return UBlueprint::StaticClass(); }
 	virtual bool HasActions ( const TArray<UObject*>& InObjects ) const override { return true; }
-	virtual void GetActions(const TArray<UObject*>& InObjects, FEditorMenuSection& Section) override;
+	virtual void GetActions(const TArray<UObject*>& InObjects, FToolMenuSection& Section) override;
 	virtual void OpenAssetEditor(const TArray<UObject*>& InObjects, TSharedPtr<class IToolkitHost> EditWithinLevelEditor = TSharedPtr<IToolkitHost>()) override;
 	virtual bool CanMerge() const override;
 	virtual void Merge(UObject* InObject) override;
diff --git a/Engine/Source/Developer/AssetTools/Public/AssetTypeActions_Base.h b/Engine/Source/Developer/AssetTools/Public/AssetTypeActions_Base.h
--- a/Engine/Source/Developer/AssetTools/Public/AssetTypeActions_Base.h
+++ b/Engine/Source/Developer/AssetTools/Public/AssetTypeActions_Base.h
@@ -18,7 +18,7 @@
 struct FAssetData;
 struct FARFilter;
 class FMenuBuilder;
-struct FEditorMenuSection;
+struct FToolMenuSection;
 
 /** A base class for all AssetTypeActions. Provides helper functions useful for many types. Deriving from this class is optional. */
 class FAssetTypeActions_Base : public IAssetTypeActions
@@ -37,7 +37,7 @@ public:
 
 	}
 
-	virtual void GetActions(const TArray<UObject*>& InObjects, FEditorMenuSection& Section) override
+	virtual void GetActions(const TArray<UObject*>& InObjects, FToolMenuSection& Section) override
 	{
 
 	}
diff --git a/Engine/Source/Developer/AssetTools/Public/IAssetTypeActions.h b/Engine/Source/Developer/AssetTools/Public/IAssetTypeActions.h
--- a/Engine/Source/Developer/AssetTools/Public/IAssetTypeActions.h
+++ b/Engine/Source/Developer/AssetTools/Public/IAssetTypeActions.h
@@ -58,7 +58,7 @@ public:
 	virtual void GetActions( const TArray<UObject*>& InObjects, class FMenuBuilder& MenuBuilder ) = 0;
 
 	/** Generates a menu section for the specified objects. */
-	virtual void GetActions(const TArray<UObject*>& InObjects, struct FEditorMenuSection& Section) = 0;
+	virtual void GetActions(const TArray<UObject*>& InObjects, struct FToolMenuSection& Section) = 0;
 
 	/** Opens the asset editor for the specified objects. If EditWithinLevelEditor is valid, the world-centric editor will be used. */
 	virtual void OpenAssetEditor( const TArray<UObject*>& InObjects, TSharedPtr<IToolkitHost> EditWithinLevelEditor = TSharedPtr<IToolkitHost>() ) = 0;
diff --git a/Engine/Source/Editor/AudioEditor/Private/AssetTypeActions/AssetTypeActions_DialogueWave.cpp b/Engine/Source/Editor/AudioEditor/Private/AssetTypeActions/AssetTypeActions_DialogueWave.cpp
--- a/Engine/Source/Editor/AudioEditor/Private/AssetTypeActions/AssetTypeActions_DialogueWave.cpp
+++ b/Engine/Source/Editor/AudioEditor/Private/AssetTypeActions/AssetTypeActions_DialogueWave.cpp
@@ -2,7 +2,7 @@
 
 #include "AssetTypeActions/AssetTypeActions_DialogueWave.h"
 #include "Misc/PackageName.h"
-#include "EditorMenuSubsystem.h"
+#include "ToolMenus.h"
 #include "Editor.h"
 #include "EditorStyleSet.h"
 #include "Factories/SoundCueFactoryNew.h"
@@ -12,7 +12,7 @@
 
 #define LOCTEXT_NAMESPACE "AssetTypeActions"
 
-void FAssetTypeActions_DialogueWave::GetActions(const TArray<UObject*>& InObjects, FEditorMenuSection& Section)
+void FAssetTypeActions_DialogueWave::GetActions(const TArray<UObject*>& InObjects, FToolMenuSection& Section)
 {
 	auto DialogueWaves = GetTypedWeakObjectPtrs<UDialogueWave>(InObjects);
 
diff --git a/Engine/Source/Editor/AudioEditor/Private/AssetTypeActions/AssetTypeActions_SoundBase.cpp b/Engine/Source/Editor/AudioEditor/Private/AssetTypeActions/AssetTypeActions_SoundBase.cpp
--- a/Engine/Source/Editor/AudioEditor/Private/AssetTypeActions/AssetTypeActions_SoundBase.cpp
+++ b/Engine/Source/Editor/AudioEditor/Private/AssetTypeActions/AssetTypeActions_SoundBase.cpp
@@ -2,7 +2,7 @@
 
 #include "AssetTypeActions/AssetTypeActions_SoundBase.h"
 #include "Sound/SoundBase.h"
-#include "EditorMenuSubsystem.h"
+#include "ToolMenus.h"
 #include "Editor.h"
 #include "EditorStyleSet.h"
 #include "Components/AudioComponent.h"
@@ -17,7 +17,7 @@ UClass* FAssetTypeActions_SoundBase::GetSupportedClass() const
 	return USoundBase::StaticClass();
 }
 
-void FAssetTypeActions_SoundBase::GetActions(const TArray<UObject*>& InObjects, FEditorMenuSection& Section)
+void FAssetTypeActions_SoundBase::GetActions(const TArray<UObject*>& InObjects, FToolMenuSection& Section)
 {
 	auto Sounds = GetTypedWeakObjectPtrs<USoundBase>(InObjects);
 
diff --git a/Engine/Source/Editor/AudioEditor/Private/AssetTypeActions/AssetTypeActions_SoundCue.cpp b/Engine/Source/Editor/AudioEditor/Private/AssetTypeActions/AssetTypeActions_SoundCue.cpp
--- a/Engine/Source/Editor/AudioEditor/Private/AssetTypeActions/AssetTypeActions_SoundCue.cpp
+++ b/Engine/Source/Editor/AudioEditor/Private/AssetTypeActions/AssetTypeActions_SoundCue.cpp
@@ -4,7 +4,7 @@
 #include "Sound/SoundAttenuation.h"
 #include "Misc/PackageName.h"
 #include "Sound/SoundCue.h"
-#include "EditorMenuSubsystem.h"
+#include "ToolMenus.h"
 #include "Factories/SoundAttenuationFactory.h"
 #include "ContentBrowserModule.h"
 #include "IContentBrowserSingleton.h"
@@ -17,7 +17,7 @@ UClass* FAssetTypeActions_SoundCue::GetSupportedClass() const
 	return USoundCue::StaticClass();
 }
 
-void FAssetTypeActions_SoundCue::GetActions(const TArray<UObject*>& InObjects, FEditorMenuSection& Section)
+void FAssetTypeActions_SoundCue::GetActions(const TArray<UObject*>& InObjects, FToolMenuSection& Section)
 {
 	auto SoundCues = GetTypedWeakObjectPtrs<USoundCue>(InObjects);
 	FAssetTypeActions_SoundBase::GetActions(InObjects, Section);
diff --git a/Engine/Source/Editor/AudioEditor/Private/AssetTypeActions/AssetTypeActions_SoundCueTemplate.cpp b/Engine/Source/Editor/AudioEditor/Private/AssetTypeActions/AssetTypeActions_SoundCueTemplate.cpp
--- a/Engine/Source/Editor/AudioEditor/Private/AssetTypeActions/AssetTypeActions_SoundCueTemplate.cpp
+++ b/Engine/Source/Editor/AudioEditor/Private/AssetTypeActions/AssetTypeActions_SoundCueTemplate.cpp
@@ -5,7 +5,7 @@
 #include "EditorStyleSet.h"
 #include "Factories/SoundCueFactoryNew.h"
 #include "Factories/SoundCueTemplateFactory.h"
-#include "EditorMenuSubsystem.h"
+#include "ToolMenus.h"
 #include "IContentBrowserSingleton.h"
 #include "Sound/SoundCue.h"
 #include "Sound/SoundCueTemplate.h"
@@ -21,7 +21,7 @@ UClass* FAssetTypeActions_SoundCueTemplate::GetSupportedClass() const
 	return USoundCueTemplate::StaticClass();
 }
 
-void FAssetTypeActions_SoundCueTemplate::GetActions(const TArray<UObject*>& InObjects, FEditorMenuSection& Section)
+void FAssetTypeActions_SoundCueTemplate::GetActions(const TArray<UObject*>& InObjects, FToolMenuSection& Section)
 {
 	TArray<TWeakObjectPtr<USoundCueTemplate>> Cues = GetTypedWeakObjectPtrs<USoundCueTemplate>(InObjects);
 
diff --git a/Engine/Source/Editor/AudioEditor/Private/AssetTypeActions/AssetTypeActions_SoundWave.cpp b/Engine/Source/Editor/AudioEditor/Private/AssetTypeActions/AssetTypeActions_SoundWave.cpp
--- a/Engine/Source/Editor/AudioEditor/Private/AssetTypeActions/AssetTypeActions_SoundWave.cpp
+++ b/Engine/Source/Editor/AudioEditor/Private/AssetTypeActions/AssetTypeActions_SoundWave.cpp
@@ -5,7 +5,7 @@
 #include "Misc/PackageName.h"
 #include "AssetData.h"
 #include "Sound/SoundWave.h"
-#include "EditorMenuSubsystem.h"
+#include "ToolMenus.h"
 #include "EditorStyleSet.h"
 #include "Factories/SoundCueFactoryNew.h"
 #include "EditorFramework/AssetImportData.h"
@@ -27,7 +27,7 @@ UClass* FAssetTypeActions_SoundWave::GetSupportedClass() const
 	return USoundWave::StaticClass();
 }
 
-void FAssetTypeActions_SoundWave::GetActions(const TArray<UObject*>& InObjects, FEditorMenuSection& Section)
+void FAssetTypeActions_SoundWave::GetActions(const TArray<UObject*>& InObjects, FToolMenuSection& Section)
 {
 	FAssetTypeActions_SoundBase::GetActions(InObjects, Section);
 
@@ -44,7 +44,7 @@ void FAssetTypeActions_SoundWave::GetActions(const TArray<UObject*>& InObjects,
 			)
 		);
 
-	Section.AddEntry(FEditorMenuEntry::InitSubMenu(
+	Section.AddEntry(FToolMenuEntry::InitSubMenu(
 		NAME_None,
 		"SoundWave_CreateDialogue",
 		LOCTEXT("SoundWave_CreateDialogue", "Create Dialogue"),
diff --git a/Engine/Source/Editor/Blutility/Private/AssetTypeActions_EditorUtilityBlueprint.cpp b/Engine/Source/Editor/Blutility/Private/AssetTypeActions_EditorUtilityBlueprint.cpp
--- a/Engine/Source/Editor/Blutility/Private/AssetTypeActions_EditorUtilityBlueprint.cpp
+++ b/Engine/Source/Editor/Blutility/Private/AssetTypeActions_EditorUtilityBlueprint.cpp
@@ -1,7 +1,7 @@
 // Copyright 1998-2019 Epic Games, Inc. All Rights Reserved.
 
 #include "AssetTypeActions_EditorUtilityBlueprint.h"
-#include "EditorMenuSubsystem.h"
+#include "ToolMenus.h"
 #include "Misc/PackageName.h"
 #include "Misc/MessageDialog.h"
 #include "EditorUtilityBlueprintFactory.h"
@@ -38,7 +38,7 @@ bool FAssetTypeActions_EditorUtilityBlueprint::HasActions(const TArray<UObject*>
 	return true;
 }
 
-void FAssetTypeActions_EditorUtilityBlueprint::GetActions(const TArray<UObject*>& InObjects, FEditorMenuSection& Section)
+void FAssetTypeActions_EditorUtilityBlueprint::GetActions(const TArray<UObject*>& InObjects, FToolMenuSection& Section)
 {
 	auto Blueprints = GetTypedWeakObjectPtrs<UEditorUtilityBlueprint>(InObjects);
 
diff --git a/Engine/Source/Editor/Blutility/Private/AssetTypeActions_EditorUtilityWidgetBlueprint.cpp b/Engine/Source/Editor/Blutility/Private/AssetTypeActions_EditorUtilityWidgetBlueprint.cpp
--- a/Engine/Source/Editor/Blutility/Private/AssetTypeActions_EditorUtilityWidgetBlueprint.cpp
+++ b/Engine/Source/Editor/Blutility/Private/AssetTypeActions_EditorUtilityWidgetBlueprint.cpp
@@ -1,7 +1,7 @@
 // Copyright 1998-2019 Epic Games, Inc. All Rights Reserved.
 
 #include "AssetTypeActions_EditorUtilityWidgetBlueprint.h"
-#include "EditorMenuSubsystem.h"
+#include "ToolMenus.h"
 #include "Misc/PackageName.h"
 #include "Misc/MessageDialog.h"
 #include "EditorUtilityBlueprintFactory.h"
@@ -45,7 +45,7 @@ bool FAssetTypeActions_EditorUtilityWidgetBlueprint::HasActions(const TArray<UOb
 	return true;
 }
 
-void FAssetTypeActions_EditorUtilityWidgetBlueprint::GetActions(const TArray<UObject*>& InObjects, FEditorMenuSection& Section)
+void FAssetTypeActions_EditorUtilityWidgetBlueprint::GetActions(const TArray<UObject*>& InObjects, FToolMenuSection& Section)
 {
 	auto Blueprints = GetTypedWeakObjectPtrs<UWidgetBlueprint>(InObjects);
 
