Rename tool menu function Find to FindContext
#rnx

diff --git a/Engine/Plugins/Media/MediaPlayerEditor/Source/MediaPlayerEditor/Private/Widgets/SMediaPlayerEditorMedia.cpp b/Engine/Plugins/Media/MediaPlayerEditor/Source/MediaPlayerEditor/Private/Widgets/SMediaPlayerEditorMedia.cpp
--- a/Engine/Plugins/Media/MediaPlayerEditor/Source/MediaPlayerEditor/Private/Widgets/SMediaPlayerEditorMedia.cpp
+++ b/Engine/Plugins/Media/MediaPlayerEditor/Source/MediaPlayerEditor/Private/Widgets/SMediaPlayerEditorMedia.cpp
@@ -139,7 +139,7 @@ void SMediaPlayerEditorMedia::RegisterMenus()
 			FSlateIcon(FEditorStyle::GetStyleSetName(), "ContentBrowser.AssetActions.Edit"),
 			FToolMenuExecuteAction::CreateLambda([](const FToolMenuContext& InContext)
 			{
-				if (UMediaPlayerEditorMediaContext* Context = InContext.Find<UMediaPlayerEditorMediaContext>())
+				if (UMediaPlayerEditorMediaContext* Context = InContext.FindContext<UMediaPlayerEditorMediaContext>())
 				{
 					if (Context->SelectedAsset)
 					{
@@ -161,7 +161,7 @@ void SMediaPlayerEditorMedia::RegisterMenus()
 					FSlateIcon(Context->StyleSetName, "MediaPlayerEditor.NextMedia.Small"),
 					FToolMenuExecuteAction::CreateLambda([=](const FToolMenuContext& InContext)
 					{
-						if (UMediaPlayerEditorMediaContext* Context = InContext.Find<UMediaPlayerEditorMediaContext>())
+						if (UMediaPlayerEditorMediaContext* Context = InContext.FindContext<UMediaPlayerEditorMediaContext>())
 						{
 							if (Context->MediaPlayerEditorMedia.IsValid())
 							{
@@ -183,7 +183,7 @@ void SMediaPlayerEditorMedia::RegisterMenus()
 			FSlateIcon(FEditorStyle::GetStyleSetName(), "SystemWideCommands.FindInContentBrowser"),
 			FToolMenuExecuteAction::CreateLambda([](const FToolMenuContext& InContext)
 			{
-				if (UMediaPlayerEditorMediaContext* Context = InContext.Find<UMediaPlayerEditorMediaContext>())
+				if (UMediaPlayerEditorMediaContext* Context = InContext.FindContext<UMediaPlayerEditorMediaContext>())
 				{
 					TArray<UObject*> AssetsToSync = TArrayBuilder<UObject*>().Add(Context->SelectedAsset);
 					GEditor->SyncBrowserToObjects(AssetsToSync);
diff --git a/Engine/Plugins/Runtime/SoundUtilities/Source/SoundUtilitiesEditor/Private/SoundWaveAssetActionExtender.cpp b/Engine/Plugins/Runtime/SoundUtilities/Source/SoundUtilitiesEditor/Private/SoundWaveAssetActionExtender.cpp
--- a/Engine/Plugins/Runtime/SoundUtilities/Source/SoundUtilitiesEditor/Private/SoundWaveAssetActionExtender.cpp
+++ b/Engine/Plugins/Runtime/SoundUtilities/Source/SoundUtilitiesEditor/Private/SoundWaveAssetActionExtender.cpp
@@ -38,7 +38,7 @@ void FSoundWaveAssetActionExtender::RegisterMenus()
 
 void FSoundWaveAssetActionExtender::ExecuteCreateSimpleSound(const FToolMenuContext& MenuContext)
 {
-	UContentBrowserAssetContextMenuContext* Context = MenuContext.Find<UContentBrowserAssetContextMenuContext>();
+	UContentBrowserAssetContextMenuContext* Context = MenuContext.FindContext<UContentBrowserAssetContextMenuContext>();
 	if (!Context || Context->SelectedObjects.Num() == 0)
 	{
 		return;
diff --git a/Engine/Plugins/Runtime/TimeSynth/Source/TimeSynthEditor/Private/TimeSynthSoundWaveAssetActionExtender.cpp b/Engine/Plugins/Runtime/TimeSynth/Source/TimeSynthEditor/Private/TimeSynthSoundWaveAssetActionExtender.cpp
--- a/Engine/Plugins/Runtime/TimeSynth/Source/TimeSynthEditor/Private/TimeSynthSoundWaveAssetActionExtender.cpp
+++ b/Engine/Plugins/Runtime/TimeSynth/Source/TimeSynthEditor/Private/TimeSynthSoundWaveAssetActionExtender.cpp
@@ -57,7 +57,7 @@ void FTimeSynthSoundWaveAssetActionExtender::RegisterMenus()
 
 void FTimeSynthSoundWaveAssetActionExtender::ExecuteCreateTimeSyncClip(const FToolMenuContext& MenuContext)
 {
-	UContentBrowserAssetContextMenuContext* Context = MenuContext.Find<UContentBrowserAssetContextMenuContext>();
+	UContentBrowserAssetContextMenuContext* Context = MenuContext.FindContext<UContentBrowserAssetContextMenuContext>();
 	if (!Context || Context->SelectedObjects.Num() == 0)
 	{
 		return;
@@ -84,7 +84,7 @@ void FTimeSynthSoundWaveAssetActionExtender::ExecuteCreateTimeSyncClip(const FTo
 
 void FTimeSynthSoundWaveAssetActionExtender::ExecuteCreateTimeSyncClipSet(const FToolMenuContext& MenuContext)
 {
-	UContentBrowserAssetContextMenuContext* Context = MenuContext.Find<UContentBrowserAssetContextMenuContext>();
+	UContentBrowserAssetContextMenuContext* Context = MenuContext.FindContext<UContentBrowserAssetContextMenuContext>();
 	if (!Context || Context->SelectedObjects.Num() == 0)
 	{
 		return;
diff --git a/Engine/Source/Developer/ToolMenus/Public/ToolMenu.h b/Engine/Source/Developer/ToolMenus/Public/ToolMenu.h
--- a/Engine/Source/Developer/ToolMenus/Public/ToolMenu.h
+++ b/Engine/Source/Developer/ToolMenus/Public/ToolMenu.h
@@ -78,7 +78,7 @@ public:
 	template <typename TContextType>
 	TContextType* FindContext() const
 	{
-		return Context.Find<TContextType>();
+		return Context.FindContext<TContextType>();
 	}
 
 	//~ Begin UObject Interface
diff --git a/Engine/Source/Developer/ToolMenus/Public/ToolMenuContext.h b/Engine/Source/Developer/ToolMenus/Public/ToolMenuContext.h
--- a/Engine/Source/Developer/ToolMenus/Public/ToolMenuContext.h
+++ b/Engine/Source/Developer/ToolMenus/Public/ToolMenuContext.h
@@ -38,7 +38,7 @@ public:
 	FToolMenuContext(TSharedPtr<FUICommandList> InCommandList, TSharedPtr<FExtender> InExtender = TSharedPtr<FExtender>(), UObject* InContext = nullptr);
 
 	template <typename TContextType>
-	TContextType* Find() const
+	TContextType* FindContext() const
 	{
 		for (UObject* Object : ContextObjects)
 		{
diff --git a/Engine/Source/Developer/ToolMenus/Public/ToolMenuSection.h b/Engine/Source/Developer/ToolMenus/Public/ToolMenuSection.h
--- a/Engine/Source/Developer/ToolMenus/Public/ToolMenuSection.h
+++ b/Engine/Source/Developer/ToolMenus/Public/ToolMenuSection.h
@@ -49,7 +49,7 @@ public:
 	template <typename TContextType>
 	TContextType* FindContext() const
 	{
-		return Context.Find<TContextType>();
+		return Context.FindContext<TContextType>();
 	}
 
 	FToolMenuEntry* FindEntry(const FName InName);
diff --git a/Engine/Source/Editor/GraphEditor/Private/SGraphEditorImpl.cpp b/Engine/Source/Editor/GraphEditor/Private/SGraphEditorImpl.cpp
--- a/Engine/Source/Editor/GraphEditor/Private/SGraphEditorImpl.cpp
+++ b/Engine/Source/Editor/GraphEditor/Private/SGraphEditorImpl.cpp
@@ -675,7 +675,7 @@ void SGraphEditorImpl::RegisterContextMenu(const UEdGraphSchema* Schema, FToolMe
 {
 	UToolMenus* ToolMenus = UToolMenus::Get();
 	const FName SchemaMenuName = Schema->GetContextMenuName();
-	UGraphNodeContextMenuContext* Context = MenuContext.Find<UGraphNodeContextMenuContext>();
+	UGraphNodeContextMenuContext* Context = MenuContext.FindContext<UGraphNodeContextMenuContext>();
 
 	// Root menu
 	// "GraphEditor.GraphContextMenu.Common"
@@ -771,7 +771,7 @@ UToolMenu* SGraphEditorImpl::GenerateContextMenu(const UEdGraphSchema* Schema, F
 	RegisterContextMenu(Schema, MenuContext);
 
 	UToolMenus* ToolMenus = UToolMenus::Get();
-	UGraphNodeContextMenuContext* Context = MenuContext.Find<UGraphNodeContextMenuContext>();
+	UGraphNodeContextMenuContext* Context = MenuContext.FindContext<UGraphNodeContextMenuContext>();
 
 	FName MenuName = NAME_None;
 	if (Context->Node)
diff --git a/Engine/Source/Editor/LevelEditor/Private/LevelEditorToolBar.cpp b/Engine/Source/Editor/LevelEditor/Private/LevelEditorToolBar.cpp
--- a/Engine/Source/Editor/LevelEditor/Private/LevelEditorToolBar.cpp
+++ b/Engine/Source/Editor/LevelEditor/Private/LevelEditorToolBar.cpp
@@ -2354,7 +2354,7 @@ void FLevelEditorToolBar::RegisterCinematicsMenu()
 	FToolMenuSection& ExistingCinematicSection = Menu->AddSection("LevelEditorExistingCinematic", LOCTEXT("CinematicMenuCombo_ExistingHeading", "Edit Existing Cinematic"));
 	ExistingCinematicSection.AddDynamicEntry("LevelEditorExistingCinematic", FNewToolMenuSectionDelegate::CreateLambda([](FToolMenuSection& InSection)
 	{
-		ULevelEditorMenuContext* FoundContext = InSection.Context.Find<ULevelEditorMenuContext>();
+		ULevelEditorMenuContext* FoundContext = InSection.Context.FindContext<ULevelEditorMenuContext>();
 		if (!FoundContext)
 		{
 			return;
