MediaPlayer: Exposed player editor viewer widget so it can be used by other modules.
#rnx

diff --git a/Engine/Plugins/Media/MediaPlayerEditor/Source/MediaPlayerEditor/Private/Widgets/SMediaPlayerEditorViewer.h b/Engine/Plugins/Media/MediaPlayerEditor/Source/MediaPlayerEditor/Private/Widgets/SMediaPlayerEditorViewer.h
deleted file mode 100644
diff --git a/Engine/Plugins/Media/MediaPlayerEditor/Source/MediaPlayerEditor/Public/Widgets/SMediaPlayerEditorViewer.h b/Engine/Plugins/Media/MediaPlayerEditor/Source/MediaPlayerEditor/Public/Widgets/SMediaPlayerEditorViewer.h
new file mode 100644
--- /dev/null
+++ b/Engine/Plugins/Media/MediaPlayerEditor/Source/MediaPlayerEditor/Public/Widgets/SMediaPlayerEditorViewer.h
@@ -0,0 +1,141 @@
+// Copyright Epic Games, Inc. All Rights Reserved.
+
+#pragma once
+
+#include "CoreTypes.h"
+#include "Internationalization/Text.h"
+#include "MediaCaptureSupport.h"
+#include "Styling/ISlateStyle.h"
+#include "Templates/SharedPointer.h"
+#include "UObject/NameTypes.h"
+#include "Widgets/DeclarativeSyntaxSupport.h"
+#include "Widgets/SCompoundWidget.h"
+
+class FMenuBuilder;
+class SEditableTextBox;
+class SSlider;
+class UMediaPlayer;
+
+enum class EMediaEvent;
+enum class EMediaPlayerTrack : uint8;
+
+
+/**
+ * Implements the contents of the viewer tab in the UMediaPlayer asset editor.
+ */
+class MEDIAPLAYEREDITOR_API SMediaPlayerEditorViewer
+	: public SCompoundWidget
+{
+public:
+
+	SLATE_BEGIN_ARGS(SMediaPlayerEditorViewer) { }
+	SLATE_END_ARGS()
+
+public:
+
+	/** Default constructor. */
+	SMediaPlayerEditorViewer();
+
+	/** Destructor. */
+	~SMediaPlayerEditorViewer();
+
+public:
+
+	/**
+	 * Construct this widget
+	 *
+	 * @param InArgs The declaration data for this widget.
+	 * @param InMediaPlayer The UMediaPlayer asset to show the details for.
+	 * @param InStyle The style set to use.
+	 */
+	void Construct(const FArguments& InArgs, UMediaPlayer& InMediaPlayer, const TSharedRef<ISlateStyle>& InStyle);
+
+public:
+
+	//~ SWidget interface
+
+	virtual void OnDragEnter(const FGeometry& MyGeometry, const FDragDropEvent& DragDropEvent) override;
+	virtual void OnDragLeave(const FDragDropEvent& DragDropEvent) override;
+	virtual FReply OnDragOver(const FGeometry& MyGeometry, const FDragDropEvent& DragDropEvent) override;
+	virtual FReply OnDrop(const FGeometry& MyGeometry, const FDragDropEvent& DragDropEvent) override;
+
+protected:
+
+	/**
+	 * Populate a menu from the given capture device information.
+	 *
+	 * @param DeviceInfos The capture device information.
+	 * @param MenuBuilder The builder for the menu.
+	 */
+	void MakeCaptureDeviceMenu(TArray<FMediaCaptureDeviceInfo>& DeviceInfos, FMenuBuilder& MenuBuilder);
+
+	/** Open the specified media URL. */
+	void OpenUrl(const FText& TextUrl);
+
+	/** Set the name of the desired native media player. */
+	void SetDesiredPlayerName(FName PlayerName);
+
+private:
+
+	/** Callback for creating the audio capture devices sub-menu. */
+	void HandleAudioCaptureDevicesMenuNewMenu(FMenuBuilder& MenuBuilder);
+
+	/** Callback for creating the player sub-menu. */
+	void HandleDecoderMenuNewMenu(FMenuBuilder& MenuBuilder);
+
+	/** Callback for getting the text of the FPS text block. */
+	FText HandleFpsTextBlockText() const;
+
+	/** Callback for creating a track format sub-menu. */
+	void HandleFormatMenuNewMenu(FMenuBuilder& MenuBuilder, EMediaPlayerTrack TrackType);
+
+	/** Callback for media player events. */
+	void HandleMediaPlayerMediaEvent(EMediaEvent Event);
+
+	/** Callback for creating the Scale sub-menu. */
+	void HandleScaleMenuNewMenu(FMenuBuilder& MenuBuilder);
+
+	/** Callback for getting the text of the timer text block. */
+	FText HandleTimerTextBlockText() const;
+
+	/** Callback for getting the tool tip of the timer text block. */
+	FText HandleTimerTextBlockToolTipText() const;
+
+	/** Callback for creating a track sub-menu. */
+	void HandleTrackMenuNewMenu(FMenuBuilder& MenuBuilder, EMediaPlayerTrack TrackType);
+
+	/** Callback for handling key down events in the URL text box. */
+	FReply HandleUrlBoxKeyDown(const FGeometry&, const FKeyEvent& KeyEvent);
+
+	/** Callback for creating the video capture devices sub-menu. */
+	void HandleVideoCaptureDevicesMenuNewMenu(FMenuBuilder& MenuBuilder);
+
+private:
+
+	/** Whether something is currently being dragged over the widget. */
+	bool DragOver;
+
+	/** Whether the dragged object is a media file that can be played. */
+	bool DragValid;
+
+	/** The text that was last typed into the URL box. */
+	FText LastUrl;
+
+	/** Pointer to the media player that is being viewed. */
+	UMediaPlayer* MediaPlayer;
+
+	/** The playback rate prior to scrubbing. */
+	float PreScrubRate;
+
+	/** Holds the scrubber slider. */
+	TSharedPtr<SSlider> ScrubberSlider;
+
+	/** The value currently being scrubbed to. */
+	float ScrubValue;
+
+	/** The style set to use for this widget. */
+	TSharedPtr<ISlateStyle> Style;
+
+	/** Media URL text box. */
+	TSharedPtr<SEditableTextBox> UrlTextBox;
+};
