MediaPlate: Player view now uses the media texture from the media plate.
#rnx

diff --git a/Engine/Plugins/Media/MediaPlate/Source/MediaPlateEditor/Private/Toolkits/MediaPlateEditorToolkit.cpp b/Engine/Plugins/Media/MediaPlate/Source/MediaPlateEditor/Private/Toolkits/MediaPlateEditorToolkit.cpp
--- a/Engine/Plugins/Media/MediaPlate/Source/MediaPlateEditor/Private/Toolkits/MediaPlateEditorToolkit.cpp
+++ b/Engine/Plugins/Media/MediaPlate/Source/MediaPlateEditor/Private/Toolkits/MediaPlateEditorToolkit.cpp
@@ -7,6 +7,7 @@
 #include "EditorReimportHandler.h"
 #include "Framework/MultiBox/MultiBoxBuilder.h"
 #include "MediaPlate.h"
+#include "MediaPlateEditorModule.h"
 #include "MediaPlayer.h"
 #include "MediaPlaylist.h"
 #include "Models/MediaPlateEditorCommands.h"
@@ -90,6 +91,13 @@ void FMediaPlateEditorToolkit::Initialize(AMediaPlate* InMediaPlate, const ETool
 	
 	ExtendToolBar();
 	RegenerateMenusAndToolbars();
+
+	// Tell the editor module that this media plate is playing.
+	FMediaPlateEditorModule* EditorModule = FModuleManager::LoadModulePtr<FMediaPlateEditorModule>("MediaPlateEditor");
+	if (EditorModule != nullptr)
+	{
+		EditorModule->MediaPlateStartedPlayback(MediaPlate);
+	}
 }
 
 /* FAssetEditorToolkit interface
@@ -323,10 +331,11 @@ float FMediaPlateEditorToolkit::GetReverseRate() const
 TSharedRef<SDockTab> FMediaPlateEditorToolkit::HandleTabManagerSpawnTab(const FSpawnTabArgs& Args, FName TabIdentifier)
 {
 	UMediaPlayer* MediaPlayer = MediaPlate->GetMediaPlayer();
+	UMediaTexture* MediaTexture = MediaPlate->GetMediaTexture();
 	TSharedPtr<SWidget> TabWidget = SNullWidget::NullWidget;
 	if (TabIdentifier == MediaPlateEditorToolkit::ViewerTabId)
 	{
-		TabWidget = SNew(SMediaPlayerEditorViewer, *MediaPlayer, Style);
+		TabWidget = SNew(SMediaPlayerEditorViewer, *MediaPlayer, MediaTexture, Style);
 	}
 
 	return SNew(SDockTab)
diff --git a/Engine/Plugins/Media/MediaPlayerEditor/Source/MediaPlayerEditor/Private/Toolkits/MediaPlayerEditorToolkit.cpp b/Engine/Plugins/Media/MediaPlayerEditor/Source/MediaPlayerEditor/Private/Toolkits/MediaPlayerEditorToolkit.cpp
--- a/Engine/Plugins/Media/MediaPlayerEditor/Source/MediaPlayerEditor/Private/Toolkits/MediaPlayerEditorToolkit.cpp
+++ b/Engine/Plugins/Media/MediaPlayerEditor/Source/MediaPlayerEditor/Private/Toolkits/MediaPlayerEditorToolkit.cpp
@@ -414,7 +414,7 @@ TSharedRef<SDockTab> FMediaPlayerEditorToolkit::HandleTabManagerSpawnTab(const F
 	}
 	else if (TabIdentifier == MediaPlayerEditorToolkit::ViewerTabId)
 	{
-		TabWidget = SNew(SMediaPlayerEditorViewer, *MediaPlayer, Style);
+		TabWidget = SNew(SMediaPlayerEditorViewer, *MediaPlayer, nullptr, Style);
 	}
 
 	return SNew(SDockTab)
diff --git a/Engine/Plugins/Media/MediaPlayerEditor/Source/MediaPlayerEditor/Private/Widgets/SMediaPlayerEditorOutput.cpp b/Engine/Plugins/Media/MediaPlayerEditor/Source/MediaPlayerEditor/Private/Widgets/SMediaPlayerEditorOutput.cpp
--- a/Engine/Plugins/Media/MediaPlayerEditor/Source/MediaPlayerEditor/Private/Widgets/SMediaPlayerEditorOutput.cpp
+++ b/Engine/Plugins/Media/MediaPlayerEditor/Source/MediaPlayerEditor/Private/Widgets/SMediaPlayerEditorOutput.cpp
@@ -23,6 +23,7 @@ SMediaPlayerEditorOutput::SMediaPlayerEditorOutput()
 	: MediaPlayer(nullptr)
 	, MediaTexture(nullptr)
 	, SoundComponent(nullptr)
+	, bIsOurMediaTexture(false)
 { }
 
 
@@ -34,10 +35,15 @@ SMediaPlayerEditorOutput::~SMediaPlayerEditorOutput()
 		MediaPlayer.Reset();
 	}
 
-	if (MediaTexture != nullptr)
+	// Did we create the media texture?
+	if (bIsOurMediaTexture)
 	{
-		MediaTexture->RemoveFromRoot();
-		MediaTexture = nullptr;
+		bIsOurMediaTexture = false;
+		if (MediaTexture != nullptr)
+		{
+			MediaTexture->RemoveFromRoot();
+			MediaTexture = nullptr;
+		}
 	}
 
 	if (SoundComponent != nullptr)
@@ -54,7 +60,8 @@ SMediaPlayerEditorOutput::~SMediaPlayerEditorOutput()
 /* SMediaPlayerEditorOutput interface
  *****************************************************************************/
 
-void SMediaPlayerEditorOutput::Construct(const FArguments& InArgs, UMediaPlayer& InMediaPlayer)
+void SMediaPlayerEditorOutput::Construct(const FArguments& InArgs, UMediaPlayer& InMediaPlayer,
+	UMediaTexture* InMediaTexture)
 {
 	MediaPlayer = &InMediaPlayer;
 
@@ -73,15 +80,25 @@ void SMediaPlayerEditorOutput::Construct(const FArguments& InArgs, UMediaPlayer&
 		}
 	}
 
-	// create media texture
-	MediaTexture = NewObject<UMediaTexture>(GetTransientPackage(), NAME_None, RF_Transient | RF_Public);
-
-	if (MediaTexture != nullptr)
+	// Did we get media texture passed in?
+	if (InMediaTexture != nullptr)
+	{
+		MediaTexture = InMediaTexture;
+		bIsOurMediaTexture = false;
+	}
+	else
 	{
-		MediaTexture->AutoClear = true;
-		MediaTexture->SetMediaPlayer(&InMediaPlayer);
-		MediaTexture->UpdateResource();
-		MediaTexture->AddToRoot();
+		// create media texture
+		MediaTexture = NewObject<UMediaTexture>(GetTransientPackage(), NAME_None, RF_Transient | RF_Public);
+		bIsOurMediaTexture = true;
+
+		if (MediaTexture != nullptr)
+		{
+			MediaTexture->AutoClear = true;
+			MediaTexture->SetMediaPlayer(&InMediaPlayer);
+			MediaTexture->UpdateResource();
+			MediaTexture->AddToRoot();
+		}
 	}
 
 	TSharedRef<SMediaImage> MediaImage = SNew(SMediaImage, MediaTexture)
diff --git a/Engine/Plugins/Media/MediaPlayerEditor/Source/MediaPlayerEditor/Private/Widgets/SMediaPlayerEditorOutput.h b/Engine/Plugins/Media/MediaPlayerEditor/Source/MediaPlayerEditor/Private/Widgets/SMediaPlayerEditorOutput.h
--- a/Engine/Plugins/Media/MediaPlayerEditor/Source/MediaPlayerEditor/Private/Widgets/SMediaPlayerEditorOutput.h
+++ b/Engine/Plugins/Media/MediaPlayerEditor/Source/MediaPlayerEditor/Private/Widgets/SMediaPlayerEditorOutput.h
@@ -40,8 +40,9 @@ public:
 	 *
 	 * @param InArgs The declaration data for this widget.
 	 * @param InMediaPlayer The UMediaPlayer asset to show the details for.
+	 * @param InMediaTexture The UMediaTexture asset to output video to. If nullptr then use our own.
 	 */
-	void Construct(const FArguments& InArgs, UMediaPlayer& InMediaPlayer);
+	void Construct(const FArguments& InArgs, UMediaPlayer& InMediaPlayer, UMediaTexture* InMediaTexture);
 
 public:
 
@@ -64,4 +65,7 @@ private:
 
 	/** The sound component to play the media player's audio output. */
 	UMediaSoundComponent* SoundComponent;
+
+	/** If true then we created MediaTexture and need to clean it up when we are done. */
+	bool bIsOurMediaTexture;
 };
diff --git a/Engine/Plugins/Media/MediaPlayerEditor/Source/MediaPlayerEditor/Private/Widgets/SMediaPlayerEditorViewer.cpp b/Engine/Plugins/Media/MediaPlayerEditor/Source/MediaPlayerEditor/Private/Widgets/SMediaPlayerEditorViewer.cpp
--- a/Engine/Plugins/Media/MediaPlayerEditor/Source/MediaPlayerEditor/Private/Widgets/SMediaPlayerEditorViewer.cpp
+++ b/Engine/Plugins/Media/MediaPlayerEditor/Source/MediaPlayerEditor/Private/Widgets/SMediaPlayerEditorViewer.cpp
@@ -63,7 +63,8 @@ SMediaPlayerEditorViewer::~SMediaPlayerEditorViewer()
 /* SMediaPlayerEditorPlayer interface
  *****************************************************************************/
 
-void SMediaPlayerEditorViewer::Construct(const FArguments& InArgs, UMediaPlayer& InMediaPlayer, const TSharedRef<ISlateStyle>& InStyle)
+void SMediaPlayerEditorViewer::Construct(const FArguments& InArgs, UMediaPlayer& InMediaPlayer,
+	UMediaTexture* InMediaTexture, const TSharedRef<ISlateStyle>& InStyle)
 {
 	MediaPlayer = &InMediaPlayer;
 	Style = InStyle;
@@ -337,7 +338,7 @@ void SMediaPlayerEditorViewer::Construct(const FArguments& InArgs, UMediaPlayer&
 							.Padding(0.0f)
 							[
 								// movie area
-								SNew(SMediaPlayerEditorViewport, InMediaPlayer, InStyle)
+								SNew(SMediaPlayerEditorViewport, InMediaPlayer, InMediaTexture, InStyle)
 							]
 					]
 
diff --git a/Engine/Plugins/Media/MediaPlayerEditor/Source/MediaPlayerEditor/Private/Widgets/SMediaPlayerEditorViewport.cpp b/Engine/Plugins/Media/MediaPlayerEditor/Source/MediaPlayerEditor/Private/Widgets/SMediaPlayerEditorViewport.cpp
--- a/Engine/Plugins/Media/MediaPlayerEditor/Source/MediaPlayerEditor/Private/Widgets/SMediaPlayerEditorViewport.cpp
+++ b/Engine/Plugins/Media/MediaPlayerEditor/Source/MediaPlayerEditor/Private/Widgets/SMediaPlayerEditorViewport.cpp
@@ -29,7 +29,8 @@ SMediaPlayerEditorViewport::SMediaPlayerEditorViewport()
 /* SMediaPlayerEditorPlayer interface
  *****************************************************************************/
 
-void SMediaPlayerEditorViewport::Construct(const FArguments& InArgs, UMediaPlayer& InMediaPlayer, const TSharedRef<ISlateStyle>& InStyle)
+void SMediaPlayerEditorViewport::Construct(const FArguments& InArgs, UMediaPlayer& InMediaPlayer,
+	UMediaTexture* InMediaTexture, const TSharedRef<ISlateStyle>& InStyle)
 {
 	MediaPlayer = &InMediaPlayer;
 	Style = InStyle;
@@ -62,7 +63,7 @@ void SMediaPlayerEditorViewport::Construct(const FArguments& InArgs, UMediaPlaye
 						})
 						[
 							// movie texture
-							SNew(SMediaPlayerEditorOutput, InMediaPlayer)
+							SNew(SMediaPlayerEditorOutput, InMediaPlayer, InMediaTexture)
 						]
 					]
 				]
diff --git a/Engine/Plugins/Media/MediaPlayerEditor/Source/MediaPlayerEditor/Private/Widgets/SMediaPlayerEditorViewport.h b/Engine/Plugins/Media/MediaPlayerEditor/Source/MediaPlayerEditor/Private/Widgets/SMediaPlayerEditorViewport.h
--- a/Engine/Plugins/Media/MediaPlayerEditor/Source/MediaPlayerEditor/Private/Widgets/SMediaPlayerEditorViewport.h
+++ b/Engine/Plugins/Media/MediaPlayerEditor/Source/MediaPlayerEditor/Private/Widgets/SMediaPlayerEditorViewport.h
@@ -7,6 +7,7 @@
 #include "Widgets/SCompoundWidget.h"
 
 class UMediaPlayer;
+class UMediaTexture;
 
 
 class SMediaPlayerEditorViewport
@@ -29,9 +30,11 @@ public:
 	 *
 	 * @param InArgs The declaration data for this widget.
 	 * @param InMediaPlayer The UMediaPlayer asset to show the details for.
+	 * @param InMediaTexture The UMediaTexture asset to output video to. If nullptr then use our own.
 	 * @param InStyleSet The style set to use.
 	 */
-	void Construct(const FArguments& InArgs, UMediaPlayer& InMediaPlayer, const TSharedRef<ISlateStyle>& InStyle);
+	void Construct(const FArguments& InArgs, UMediaPlayer& InMediaPlayer,
+		UMediaTexture* InMediaTexture, const TSharedRef<ISlateStyle>& InStyle);
 
 public:
 
diff --git a/Engine/Plugins/Media/MediaPlayerEditor/Source/MediaPlayerEditor/Public/Widgets/SMediaPlayerEditorViewer.h b/Engine/Plugins/Media/MediaPlayerEditor/Source/MediaPlayerEditor/Public/Widgets/SMediaPlayerEditorViewer.h
--- a/Engine/Plugins/Media/MediaPlayerEditor/Source/MediaPlayerEditor/Public/Widgets/SMediaPlayerEditorViewer.h
+++ b/Engine/Plugins/Media/MediaPlayerEditor/Source/MediaPlayerEditor/Public/Widgets/SMediaPlayerEditorViewer.h
@@ -15,6 +15,7 @@ class FMenuBuilder;
 class SEditableTextBox;
 class SSlider;
 class UMediaPlayer;
+class UMediaTexture;
 
 enum class EMediaEvent;
 enum class EMediaPlayerTrack : uint8;
@@ -46,9 +47,11 @@ public:
 	 *
 	 * @param InArgs The declaration data for this widget.
 	 * @param InMediaPlayer The UMediaPlayer asset to show the details for.
+	 * @param InMediaTexture The UMediaTexture asset to output video to. If nullptr then use our own.
 	 * @param InStyle The style set to use.
 	 */
-	void Construct(const FArguments& InArgs, UMediaPlayer& InMediaPlayer, const TSharedRef<ISlateStyle>& InStyle);
+	void Construct(const FArguments& InArgs, UMediaPlayer& InMediaPlayer,
+		UMediaTexture* InMediaTexture, const TSharedRef<ISlateStyle>& InStyle);
 
 public:
 
