MediaPlayerEditor: Only creates sound component if needed.

#rnx

diff --git a/Engine/Plugins/Media/MediaPlate/Source/MediaPlateEditor/Private/Toolkits/MediaPlateEditorToolkit.cpp b/Engine/Plugins/Media/MediaPlate/Source/MediaPlateEditor/Private/Toolkits/MediaPlateEditorToolkit.cpp
--- a/Engine/Plugins/Media/MediaPlate/Source/MediaPlateEditor/Private/Toolkits/MediaPlateEditorToolkit.cpp
+++ b/Engine/Plugins/Media/MediaPlate/Source/MediaPlateEditor/Private/Toolkits/MediaPlateEditorToolkit.cpp
@@ -358,7 +358,7 @@ TSharedRef<SDockTab> FMediaPlateEditorToolkit::HandleTabManagerSpawnTab(const FS
 	}
 	else if (TabIdentifier == MediaPlateEditorToolkit::ViewerTabId)
 	{
-		TabWidget = SNew(SMediaPlayerEditorViewer, *MediaPlayer, MediaTexture, Style);
+		TabWidget = SNew(SMediaPlayerEditorViewer, *MediaPlayer, MediaTexture, Style, false);
 	}
 
 	return SNew(SDockTab)
diff --git a/Engine/Plugins/Media/MediaPlayerEditor/Source/MediaPlayerEditor/Private/Toolkits/MediaPlayerEditorToolkit.cpp b/Engine/Plugins/Media/MediaPlayerEditor/Source/MediaPlayerEditor/Private/Toolkits/MediaPlayerEditorToolkit.cpp
--- a/Engine/Plugins/Media/MediaPlayerEditor/Source/MediaPlayerEditor/Private/Toolkits/MediaPlayerEditorToolkit.cpp
+++ b/Engine/Plugins/Media/MediaPlayerEditor/Source/MediaPlayerEditor/Private/Toolkits/MediaPlayerEditorToolkit.cpp
@@ -414,7 +414,7 @@ TSharedRef<SDockTab> FMediaPlayerEditorToolkit::HandleTabManagerSpawnTab(const F
 	}
 	else if (TabIdentifier == MediaPlayerEditorToolkit::ViewerTabId)
 	{
-		TabWidget = SNew(SMediaPlayerEditorViewer, *MediaPlayer, nullptr, Style);
+		TabWidget = SNew(SMediaPlayerEditorViewer, *MediaPlayer, nullptr, Style, true);
 	}
 
 	return SNew(SDockTab)
diff --git a/Engine/Plugins/Media/MediaPlayerEditor/Source/MediaPlayerEditor/Private/Widgets/SMediaPlayerEditorOutput.cpp b/Engine/Plugins/Media/MediaPlayerEditor/Source/MediaPlayerEditor/Private/Widgets/SMediaPlayerEditorOutput.cpp
--- a/Engine/Plugins/Media/MediaPlayerEditor/Source/MediaPlayerEditor/Private/Widgets/SMediaPlayerEditorOutput.cpp
+++ b/Engine/Plugins/Media/MediaPlayerEditor/Source/MediaPlayerEditor/Private/Widgets/SMediaPlayerEditorOutput.cpp
@@ -61,12 +61,12 @@ SMediaPlayerEditorOutput::~SMediaPlayerEditorOutput()
  *****************************************************************************/
 
 void SMediaPlayerEditorOutput::Construct(const FArguments& InArgs, UMediaPlayer& InMediaPlayer,
-	UMediaTexture* InMediaTexture)
+	UMediaTexture* InMediaTexture, bool bInIsSoundEnabled)
 {
 	MediaPlayer = &InMediaPlayer;
 
 	// create media sound component
-	if ((GEngine != nullptr) && GEngine->UseSound())
+	if ((GEngine != nullptr) && GEngine->UseSound() && (bInIsSoundEnabled))
 	{
 		SoundComponent = NewObject<UMediaSoundComponent>(GetTransientPackage(), NAME_None, RF_Transient | RF_Public);
 
diff --git a/Engine/Plugins/Media/MediaPlayerEditor/Source/MediaPlayerEditor/Private/Widgets/SMediaPlayerEditorOutput.h b/Engine/Plugins/Media/MediaPlayerEditor/Source/MediaPlayerEditor/Private/Widgets/SMediaPlayerEditorOutput.h
--- a/Engine/Plugins/Media/MediaPlayerEditor/Source/MediaPlayerEditor/Private/Widgets/SMediaPlayerEditorOutput.h
+++ b/Engine/Plugins/Media/MediaPlayerEditor/Source/MediaPlayerEditor/Private/Widgets/SMediaPlayerEditorOutput.h
@@ -41,8 +41,10 @@ public:
 	 * @param InArgs The declaration data for this widget.
 	 * @param InMediaPlayer The UMediaPlayer asset to show the details for.
 	 * @param InMediaTexture The UMediaTexture asset to output video to. If nullptr then use our own.
+	 * @param bInIsSoundEnabled If true then produce sound.
 	 */
-	void Construct(const FArguments& InArgs, UMediaPlayer& InMediaPlayer, UMediaTexture* InMediaTexture);
+	void Construct(const FArguments& InArgs, UMediaPlayer& InMediaPlayer,
+		UMediaTexture* InMediaTexture, bool  bInIsSoundEnabled);
 
 public:
 
diff --git a/Engine/Plugins/Media/MediaPlayerEditor/Source/MediaPlayerEditor/Private/Widgets/SMediaPlayerEditorViewer.cpp b/Engine/Plugins/Media/MediaPlayerEditor/Source/MediaPlayerEditor/Private/Widgets/SMediaPlayerEditorViewer.cpp
--- a/Engine/Plugins/Media/MediaPlayerEditor/Source/MediaPlayerEditor/Private/Widgets/SMediaPlayerEditorViewer.cpp
+++ b/Engine/Plugins/Media/MediaPlayerEditor/Source/MediaPlayerEditor/Private/Widgets/SMediaPlayerEditorViewer.cpp
@@ -64,7 +64,8 @@ SMediaPlayerEditorViewer::~SMediaPlayerEditorViewer()
  *****************************************************************************/
 
 void SMediaPlayerEditorViewer::Construct(const FArguments& InArgs, UMediaPlayer& InMediaPlayer,
-	UMediaTexture* InMediaTexture, const TSharedRef<ISlateStyle>& InStyle)
+	UMediaTexture* InMediaTexture, const TSharedRef<ISlateStyle>& InStyle,
+	bool bInIsSoundEnabled)
 {
 	MediaPlayer = &InMediaPlayer;
 	Style = InStyle;
@@ -338,7 +339,7 @@ void SMediaPlayerEditorViewer::Construct(const FArguments& InArgs, UMediaPlayer&
 							.Padding(0.0f)
 							[
 								// movie area
-								SNew(SMediaPlayerEditorViewport, InMediaPlayer, InMediaTexture, InStyle)
+								SNew(SMediaPlayerEditorViewport, InMediaPlayer, InMediaTexture, InStyle, bInIsSoundEnabled)
 							]
 					]
 
diff --git a/Engine/Plugins/Media/MediaPlayerEditor/Source/MediaPlayerEditor/Private/Widgets/SMediaPlayerEditorViewport.cpp b/Engine/Plugins/Media/MediaPlayerEditor/Source/MediaPlayerEditor/Private/Widgets/SMediaPlayerEditorViewport.cpp
--- a/Engine/Plugins/Media/MediaPlayerEditor/Source/MediaPlayerEditor/Private/Widgets/SMediaPlayerEditorViewport.cpp
+++ b/Engine/Plugins/Media/MediaPlayerEditor/Source/MediaPlayerEditor/Private/Widgets/SMediaPlayerEditorViewport.cpp
@@ -30,7 +30,8 @@ SMediaPlayerEditorViewport::SMediaPlayerEditorViewport()
  *****************************************************************************/
 
 void SMediaPlayerEditorViewport::Construct(const FArguments& InArgs, UMediaPlayer& InMediaPlayer,
-	UMediaTexture* InMediaTexture, const TSharedRef<ISlateStyle>& InStyle)
+	UMediaTexture* InMediaTexture, const TSharedRef<ISlateStyle>& InStyle,
+	bool bInIsSoundEnabled)
 {
 	MediaPlayer = &InMediaPlayer;
 	Style = InStyle;
@@ -63,7 +64,7 @@ void SMediaPlayerEditorViewport::Construct(const FArguments& InArgs, UMediaPlaye
 						})
 						[
 							// movie texture
-							SNew(SMediaPlayerEditorOutput, InMediaPlayer, InMediaTexture)
+							SNew(SMediaPlayerEditorOutput, InMediaPlayer, InMediaTexture, bInIsSoundEnabled)
 						]
 					]
 				]
diff --git a/Engine/Plugins/Media/MediaPlayerEditor/Source/MediaPlayerEditor/Private/Widgets/SMediaPlayerEditorViewport.h b/Engine/Plugins/Media/MediaPlayerEditor/Source/MediaPlayerEditor/Private/Widgets/SMediaPlayerEditorViewport.h
--- a/Engine/Plugins/Media/MediaPlayerEditor/Source/MediaPlayerEditor/Private/Widgets/SMediaPlayerEditorViewport.h
+++ b/Engine/Plugins/Media/MediaPlayerEditor/Source/MediaPlayerEditor/Private/Widgets/SMediaPlayerEditorViewport.h
@@ -7,6 +7,7 @@
 #include "Widgets/SCompoundWidget.h"
 
 class UMediaPlayer;
+class UMediaSoundComponent;
 class UMediaTexture;
 
 
@@ -32,9 +33,11 @@ public:
 	 * @param InMediaPlayer The UMediaPlayer asset to show the details for.
 	 * @param InMediaTexture The UMediaTexture asset to output video to. If nullptr then use our own.
 	 * @param InStyleSet The style set to use.
+	 * @param bInIsSoundEnabled If true then produce sound.
 	 */
 	void Construct(const FArguments& InArgs, UMediaPlayer& InMediaPlayer,
-		UMediaTexture* InMediaTexture, const TSharedRef<ISlateStyle>& InStyle);
+		UMediaTexture* InMediaTexture, const TSharedRef<ISlateStyle>& InStyle,
+		bool bInIsSoundEnabled);
 
 public:
 
diff --git a/Engine/Plugins/Media/MediaPlayerEditor/Source/MediaPlayerEditor/Public/Widgets/SMediaPlayerEditorViewer.h b/Engine/Plugins/Media/MediaPlayerEditor/Source/MediaPlayerEditor/Public/Widgets/SMediaPlayerEditorViewer.h
--- a/Engine/Plugins/Media/MediaPlayerEditor/Source/MediaPlayerEditor/Public/Widgets/SMediaPlayerEditorViewer.h
+++ b/Engine/Plugins/Media/MediaPlayerEditor/Source/MediaPlayerEditor/Public/Widgets/SMediaPlayerEditorViewer.h
@@ -15,6 +15,7 @@ class FMenuBuilder;
 class SEditableTextBox;
 class SSlider;
 class UMediaPlayer;
+class UMediaSoundComponent;
 class UMediaTexture;
 
 enum class EMediaEvent;
@@ -49,9 +50,11 @@ public:
 	 * @param InMediaPlayer The UMediaPlayer asset to show the details for.
 	 * @param InMediaTexture The UMediaTexture asset to output video to. If nullptr then use our own.
 	 * @param InStyle The style set to use.
+	 * @param bInIsSoundEnabled If true then produce sound.
 	 */
 	void Construct(const FArguments& InArgs, UMediaPlayer& InMediaPlayer,
-		UMediaTexture* InMediaTexture, const TSharedRef<ISlateStyle>& InStyle);
+		UMediaTexture* InMediaTexture, const TSharedRef<ISlateStyle>& InStyle,
+		bool bInIsSoundEnabled);
 
 public:
 
