MediaPlate: Added SMediaPlateEditorPlaylist so it can detect when the playlist changes in the media plate and refresh.

#rnx

diff --git a/Engine/Plugins/Media/MediaPlate/Source/MediaPlateEditor/Private/Toolkits/MediaPlateEditorToolkit.cpp b/Engine/Plugins/Media/MediaPlate/Source/MediaPlateEditor/Private/Toolkits/MediaPlateEditorToolkit.cpp
--- a/Engine/Plugins/Media/MediaPlate/Source/MediaPlateEditor/Private/Toolkits/MediaPlateEditorToolkit.cpp
+++ b/Engine/Plugins/Media/MediaPlate/Source/MediaPlateEditor/Private/Toolkits/MediaPlateEditorToolkit.cpp
@@ -16,8 +16,8 @@
 #include "Widgets/Docking/SDockTab.h"
 #include "Widgets/SMediaPlateEditorDetails.h"
 #include "Widgets/SMediaPlateEditorMediaDetails.h"
+#include "Widgets/SMediaPlateEditorPlaylist.h"
 #include "Widgets/SMediaPlayerEditorViewer.h"
-#include "Widgets/SMediaPlaylistEditorTracks.h"
 
 #define LOCTEXT_NAMESPACE "FMediaPlateEditorToolkit"
 
@@ -401,7 +401,7 @@ TSharedRef<SDockTab> FMediaPlateEditorToolkit::HandleTabManagerSpawnTab(const FS
 	}
 	else if (TabIdentifier == MediaPlateEditorToolkit::PlaylistTabId)
 	{
-		TabWidget = SNew(SMediaPlaylistEditorTracks, MediaPlate->MediaPlaylist, Style);
+		TabWidget = SNew(SMediaPlateEditorPlaylist, *MediaPlate, Style);
 	}
 	else if (TabIdentifier == MediaPlateEditorToolkit::ViewerTabId)
 	{
diff --git a/Engine/Plugins/Media/MediaPlate/Source/MediaPlateEditor/Private/Widgets/SMediaPlateEditorPlaylist.cpp b/Engine/Plugins/Media/MediaPlate/Source/MediaPlateEditor/Private/Widgets/SMediaPlateEditorPlaylist.cpp
new file mode 100644
--- /dev/null
+++ b/Engine/Plugins/Media/MediaPlate/Source/MediaPlateEditor/Private/Widgets/SMediaPlateEditorPlaylist.cpp
@@ -0,0 +1,34 @@
+// Copyright Epic Games, Inc. All Rights Reserved.
+
+#include "Widgets/SMediaPlateEditorPlaylist.h"
+#include "MediaPlateComponent.h"
+#include "MediaPlaylist.h"
+
+#define LOCTEXT_NAMESPACE "SMediaPlateEditorPlaylist"
+
+void SMediaPlateEditorPlaylist::Construct(const FArguments& InArgs, UMediaPlateComponent& InMediaPlate, const TSharedRef<ISlateStyle>& InStyle)
+{
+	MediaPlate = &InMediaPlate;
+
+	SMediaPlaylistEditorTracks::Construct(SMediaPlaylistEditorTracks::FArguments(), MediaPlate->MediaPlaylist, InStyle);
+}
+
+void SMediaPlateEditorPlaylist::Tick(const FGeometry& AllottedGeometry, const double InCurrentTime, const float InDeltaTime)
+{
+	// Call parent.
+	SMediaPlaylistEditorTracks::Tick(AllottedGeometry, InCurrentTime, InDeltaTime);
+
+	if (MediaPlate != nullptr)
+	{
+		// Did the playlist change?
+		UMediaPlaylist* MediaPlaylist = MediaPlaylistPtr.Get();
+		if (MediaPlaylist != MediaPlate->MediaPlaylist)
+		{
+			// Refresh.
+			MediaPlaylistPtr = MediaPlate->MediaPlaylist;
+			RefreshPlaylist();
+		}
+	}
+}
+
+#undef LOCTEXT_NAMESPACE
diff --git a/Engine/Plugins/Media/MediaPlate/Source/MediaPlateEditor/Private/Widgets/SMediaPlateEditorPlaylist.h b/Engine/Plugins/Media/MediaPlate/Source/MediaPlateEditor/Private/Widgets/SMediaPlateEditorPlaylist.h
new file mode 100644
--- /dev/null
+++ b/Engine/Plugins/Media/MediaPlate/Source/MediaPlateEditor/Private/Widgets/SMediaPlateEditorPlaylist.h
@@ -0,0 +1,38 @@
+// Copyright Epic Games, Inc. All Rights Reserved.
+
+#pragma once
+
+#include "Widgets/SMediaPlaylistEditorTracks.h"
+
+class UMediaPlateComponent;
+
+/**
+ * Implements the playlist panel of the MediaPlate asset editor.
+ */
+class SMediaPlateEditorPlaylist
+	: public SMediaPlaylistEditorTracks
+{
+public:
+
+	SLATE_BEGIN_ARGS(SMediaPlateEditorPlaylist) { }
+	SLATE_END_ARGS()
+
+public:
+
+	/**
+	 * Construct this widget.
+	 *
+	 * @param InArgs		The declaration data for this widget.
+	 * @param InMediaPlate	The MediaPlate to show the details for.
+	 * @param InStyleSet	The style set to use.
+	 */
+	void Construct(const FArguments& InArgs, UMediaPlateComponent& InMediaPlate, const TSharedRef<ISlateStyle>& InStyle);
+
+	//~ SWidget interface
+	virtual void Tick(const FGeometry& AllottedGeometry, const double InCurrentTime, const float InDeltaTime) override;
+
+private:
+
+	/** Pointer to the MediaPlate that is being viewed. */
+	TWeakObjectPtr<UMediaPlateComponent> MediaPlate;
+};
diff --git a/Engine/Plugins/Media/MediaPlayerEditor/Source/MediaPlayerEditor/Private/Widgets/SMediaPlaylistEditorTracks.cpp b/Engine/Plugins/Media/MediaPlayerEditor/Source/MediaPlayerEditor/Private/Widgets/SMediaPlaylistEditorTracks.cpp
--- a/Engine/Plugins/Media/MediaPlayerEditor/Source/MediaPlayerEditor/Private/Widgets/SMediaPlaylistEditorTracks.cpp
+++ b/Engine/Plugins/Media/MediaPlayerEditor/Source/MediaPlayerEditor/Private/Widgets/SMediaPlaylistEditorTracks.cpp
@@ -13,7 +13,7 @@
 
 void SMediaPlaylistEditorTracks::Construct(const FArguments& InArgs, UMediaPlaylist* InMediaPlaylist, const TSharedRef<ISlateStyle>& InStyle)
 {
-	MediaPlaylist = InMediaPlaylist;
+	MediaPlaylistPtr = InMediaPlaylist;
 
 	ChildSlot
 		[
@@ -63,6 +63,7 @@ void SMediaPlaylistEditorTracks::RefreshPlaylist()
 {
 	SourcesContainer->ClearChildren();
 
+	UMediaPlaylist* MediaPlaylist = MediaPlaylistPtr.Get();
 	if (MediaPlaylist != nullptr)
 	{
 		// Add each item in the playlist.
@@ -95,6 +96,7 @@ void SMediaPlaylistEditorTracks::RefreshPlaylist()
 
 void SMediaPlaylistEditorTracks::AddToPlaylist()
 {
+	UMediaPlaylist* MediaPlaylist = MediaPlaylistPtr.Get();
 	if (MediaPlaylist != nullptr)
 	{
 		MediaPlaylist->Insert(nullptr, MediaPlaylist->Num());
@@ -107,6 +109,7 @@ FString SMediaPlaylistEditorTracks::GetMediaSourcePath(int32 Index) const
 {
 	FString Path;
 
+	UMediaPlaylist* MediaPlaylist = MediaPlaylistPtr.Get();
 	if (MediaPlaylist != nullptr)
 	{
 		UMediaSource* MediaSource = MediaPlaylist->Get(Index);
@@ -121,6 +124,7 @@ FString SMediaPlaylistEditorTracks::GetMediaSourcePath(int32 Index) const
 
 void SMediaPlaylistEditorTracks::OnMediaSourceChanged(const FAssetData& AssetData, int32 Index)
 {
+	UMediaPlaylist* MediaPlaylist = MediaPlaylistPtr.Get();
 	if (MediaPlaylist != nullptr)
 	{
 		UMediaSource* MediaSource = Cast<UMediaSource>(AssetData.GetAsset());
diff --git a/Engine/Plugins/Media/MediaPlayerEditor/Source/MediaPlayerEditor/Public/Widgets/SMediaPlaylistEditorTracks.h b/Engine/Plugins/Media/MediaPlayerEditor/Source/MediaPlayerEditor/Public/Widgets/SMediaPlaylistEditorTracks.h
--- a/Engine/Plugins/Media/MediaPlayerEditor/Source/MediaPlayerEditor/Public/Widgets/SMediaPlaylistEditorTracks.h
+++ b/Engine/Plugins/Media/MediaPlayerEditor/Source/MediaPlayerEditor/Public/Widgets/SMediaPlaylistEditorTracks.h
@@ -32,10 +32,10 @@ public:
 	 */
 	void Construct(const FArguments& InArgs, UMediaPlaylist* InMediaPlaylist, const TSharedRef<ISlateStyle>& InStyle);
 
-private:
+protected:
 
 	/** Pointer to the MediaPlaylist asset that is being viewed. */
-	TObjectPtr<UMediaPlaylist> MediaPlaylist;
+	TWeakObjectPtr<UMediaPlaylist> MediaPlaylistPtr;
 	/** Pointer to the container for the media sources. */
 	TSharedPtr<SVerticalBox> SourcesContainer;
 
