MediaPlate: Added scrollbox to playlist track widget so the whole playlist can be viewed.
#rnx

diff --git a/Engine/Plugins/Media/MediaPlayerEditor/Source/MediaPlayerEditor/Private/Widgets/SMediaPlaylistEditorTracks.cpp b/Engine/Plugins/Media/MediaPlayerEditor/Source/MediaPlayerEditor/Private/Widgets/SMediaPlaylistEditorTracks.cpp
--- a/Engine/Plugins/Media/MediaPlayerEditor/Source/MediaPlayerEditor/Private/Widgets/SMediaPlaylistEditorTracks.cpp
+++ b/Engine/Plugins/Media/MediaPlayerEditor/Source/MediaPlayerEditor/Private/Widgets/SMediaPlaylistEditorTracks.cpp
@@ -8,6 +8,7 @@
 #include "MediaPlaylist.h"
 #include "MediaSource.h"
 #include "Widgets/Input/SButton.h"
+#include "Widgets/Layout/SScrollBox.h"
 #include "Styling/AppStyle.h"
 
 #define LOCTEXT_NAMESPACE "SMediaPlaylistEditorTracks"
@@ -18,11 +19,10 @@ void SMediaPlaylistEditorTracks::Construct(const FArguments& InArgs, UMediaPlayl
 
 	ChildSlot
 		[
-			SNew(SVerticalBox)
+			SNew(SScrollBox)
 
 			// Buttons to manipulate playlist.
-			+ SVerticalBox::Slot()
-				.AutoHeight()
+			+ SScrollBox::Slot()
 				.Padding(2)
 				.VAlign(VAlign_Center)
 				[
@@ -51,7 +51,7 @@ void SMediaPlaylistEditorTracks::Construct(const FArguments& InArgs, UMediaPlayl
 				]
 
 			// Container for our media sources.
-			+ SVerticalBox::Slot()
+			+ SScrollBox::Slot()
 				[
 					SAssignNew(SourcesContainer, SVerticalBox)
 				]
