MediaPlate: Fixed comment in SMediaPlaylistEditorTracks.
#rnx

diff --git a/Engine/Plugins/Media/MediaPlayerEditor/Source/MediaPlayerEditor/Private/Widgets/SMediaPlaylistEditorTracks.cpp b/Engine/Plugins/Media/MediaPlayerEditor/Source/MediaPlayerEditor/Private/Widgets/SMediaPlaylistEditorTracks.cpp
--- a/Engine/Plugins/Media/MediaPlayerEditor/Source/MediaPlayerEditor/Private/Widgets/SMediaPlaylistEditorTracks.cpp
+++ b/Engine/Plugins/Media/MediaPlayerEditor/Source/MediaPlayerEditor/Private/Widgets/SMediaPlaylistEditorTracks.cpp
@@ -35,7 +35,7 @@ void SMediaPlaylistEditorTracks::Construct(const FArguments& InArgs, UMediaPlayl
 						[
 							// Add button to add to the playlist.
 							SNew(SButton)
-							.ToolTipText(LOCTEXT("MediaSource_ToolTip", "The Media Source to play."))
+							.ToolTipText(LOCTEXT("Add_ToolTip", "Add an entry to the playlist."))
 								.VAlign(VAlign_Center)
 								.OnClicked_Lambda([this]() -> FReply
 								{
