MediaPlayer: Added insert, delete, and duplicate to the playlist track editor.
#rnx

diff --git a/Engine/Plugins/Media/MediaPlayerEditor/Source/MediaPlayerEditor/Private/Widgets/SMediaPlaylistEditorTracks.cpp b/Engine/Plugins/Media/MediaPlayerEditor/Source/MediaPlayerEditor/Private/Widgets/SMediaPlaylistEditorTracks.cpp
--- a/Engine/Plugins/Media/MediaPlayerEditor/Source/MediaPlayerEditor/Private/Widgets/SMediaPlaylistEditorTracks.cpp
+++ b/Engine/Plugins/Media/MediaPlayerEditor/Source/MediaPlayerEditor/Private/Widgets/SMediaPlaylistEditorTracks.cpp
@@ -80,7 +80,7 @@ void SMediaPlaylistEditorTracks::RefreshPlaylist()
 					SNew(SHorizontalBox)
 					
 					+ SHorizontalBox::Slot()
-						.FillWidth(0.11f)
+						.AutoWidth()
 						.Padding(2)
 						.HAlign(HAlign_Left)
 						[
@@ -90,6 +90,44 @@ void SMediaPlaylistEditorTracks::RefreshPlaylist()
 								.ObjectPath(this, &SMediaPlaylistEditorTracks::GetMediaSourcePath, Index)
 								.OnObjectChanged(this, &SMediaPlaylistEditorTracks::OnMediaSourceChanged, Index)
 						]
+
+					+ SHorizontalBox::Slot()
+						.AutoWidth()
+						.Padding(2)
+						.HAlign(HAlign_Left)
+						[
+							PropertyCustomizationHelpers::MakeInsertDeleteDuplicateButton(
+								FExecuteAction::CreateLambda([this, Index]()
+								{
+									UMediaPlaylist* MediaPlaylist = MediaPlaylistPtr.Get();
+									if (MediaPlaylist != nullptr)
+									{
+										MediaPlaylist->Insert(nullptr, Index);
+										MediaPlaylist->MarkPackageDirty();
+										RefreshPlaylist();
+									}
+								}),
+								FExecuteAction::CreateLambda([this, Index]()
+								{
+									UMediaPlaylist* MediaPlaylist = MediaPlaylistPtr.Get();
+									if (MediaPlaylist != nullptr)
+									{
+										MediaPlaylist->RemoveAt(Index);
+										MediaPlaylist->MarkPackageDirty();
+										RefreshPlaylist();
+									}
+								}),
+								FExecuteAction::CreateLambda([this, Index]()
+								{
+									UMediaPlaylist* MediaPlaylist = MediaPlaylistPtr.Get();
+									if (MediaPlaylist != nullptr)
+									{
+										MediaPlaylist->Insert(MediaPlaylist->Get(Index), Index);
+										MediaPlaylist->MarkPackageDirty();
+										RefreshPlaylist();
+									}
+								}))
+						]
 				];
 		}
 	}
