MediaPlayer: Can disable SMediaPlayerEditorViewport capturing mouse events.
#rnx

diff --git a/Engine/Plugins/Media/MediaPlayerEditor/Source/MediaPlayerEditor/Private/Widgets/SMediaPlayerEditorViewer.cpp b/Engine/Plugins/Media/MediaPlayerEditor/Source/MediaPlayerEditor/Private/Widgets/SMediaPlayerEditorViewer.cpp
--- a/Engine/Plugins/Media/MediaPlayerEditor/Source/MediaPlayerEditor/Private/Widgets/SMediaPlayerEditorViewer.cpp
+++ b/Engine/Plugins/Media/MediaPlayerEditor/Source/MediaPlayerEditor/Private/Widgets/SMediaPlayerEditorViewer.cpp
@@ -340,7 +340,7 @@ void SMediaPlayerEditorViewer::Construct(const FArguments& InArgs, UMediaPlayer&
 							.Padding(0.0f)
 							[
 								// movie area
-								SNew(SMediaPlayerEditorViewport, InMediaPlayer, InMediaTexture, InStyle, bInIsSoundEnabled)
+								SAssignNew(PlayerViewport, SMediaPlayerEditorViewport, InMediaPlayer, InMediaTexture, InStyle, bInIsSoundEnabled)
 							]
 					]
 
@@ -549,6 +549,13 @@ void SMediaPlayerEditorViewer::Construct(const FArguments& InArgs, UMediaPlayer&
 	];
 }
 
+void SMediaPlayerEditorViewer::EnableMouseControl(bool bIsEnabled)
+{
+	if (PlayerViewport.IsValid())
+	{
+		PlayerViewport->EnableMouseControl(bIsEnabled);
+	}
+}
 
 /* SWidget interface
  *****************************************************************************/
diff --git a/Engine/Plugins/Media/MediaPlayerEditor/Source/MediaPlayerEditor/Private/Widgets/SMediaPlayerEditorViewport.cpp b/Engine/Plugins/Media/MediaPlayerEditor/Source/MediaPlayerEditor/Private/Widgets/SMediaPlayerEditorViewport.cpp
--- a/Engine/Plugins/Media/MediaPlayerEditor/Source/MediaPlayerEditor/Private/Widgets/SMediaPlayerEditorViewport.cpp
+++ b/Engine/Plugins/Media/MediaPlayerEditor/Source/MediaPlayerEditor/Private/Widgets/SMediaPlayerEditorViewport.cpp
@@ -23,6 +23,7 @@
 
 SMediaPlayerEditorViewport::SMediaPlayerEditorViewport()
 	: MediaPlayer(nullptr)
+	, bIsMouseControlEnabled(true)
 { }
 
 
@@ -171,7 +172,8 @@ void SMediaPlayerEditorViewport::Construct(const FArguments& InArgs, UMediaPlaye
 
 FReply SMediaPlayerEditorViewport::OnMouseButtonDown(const FGeometry& MyGeometry, const FPointerEvent& MouseEvent)
 {
-	if (MouseEvent.GetEffectingButton() != EKeys::LeftMouseButton)
+	if ((bIsMouseControlEnabled == false) ||
+		(MouseEvent.GetEffectingButton() != EKeys::LeftMouseButton))
 	{
 		return FReply::Unhandled();
 	}
@@ -184,7 +186,8 @@ FReply SMediaPlayerEditorViewport::OnMouseButtonDown(const FGeometry& MyGeometry
 
 FReply SMediaPlayerEditorViewport::OnMouseButtonDoubleClick(const FGeometry& MyGeometry, const FPointerEvent& MouseEvent)
 {
-	if (MouseEvent.GetEffectingButton() != EKeys::LeftMouseButton)
+	if ((bIsMouseControlEnabled == false) ||
+		(MouseEvent.GetEffectingButton() != EKeys::LeftMouseButton))
 	{
 		return FReply::Unhandled();
 	}
@@ -204,7 +207,8 @@ FReply SMediaPlayerEditorViewport::OnMouseButtonDoubleClick(const FGeometry& MyG
 
 FReply SMediaPlayerEditorViewport::OnMouseButtonUp(const FGeometry& MyGeometry, const FPointerEvent& MouseEvent)
 {
-	if (MouseEvent.GetEffectingButton() != EKeys::LeftMouseButton)
+	if ((bIsMouseControlEnabled == false) ||
+		(MouseEvent.GetEffectingButton() != EKeys::LeftMouseButton))
 	{
 		return FReply::Unhandled();
 	}
@@ -231,6 +235,11 @@ FReply SMediaPlayerEditorViewport::OnMouseMove(const FGeometry& MyGeometry, cons
 
 FReply SMediaPlayerEditorViewport::OnMouseWheel(const FGeometry& MyGeometry, const FPointerEvent& MouseEvent)
 {
+	if (bIsMouseControlEnabled == false)
+	{
+		return FReply::Unhandled();
+	}
+
 	MediaPlayer->SetViewField(MouseEvent.GetWheelDelta(), 0.0f, false);
 
 	return FReply::Handled();
diff --git a/Engine/Plugins/Media/MediaPlayerEditor/Source/MediaPlayerEditor/Private/Widgets/SMediaPlayerEditorViewport.h b/Engine/Plugins/Media/MediaPlayerEditor/Source/MediaPlayerEditor/Private/Widgets/SMediaPlayerEditorViewport.h
--- a/Engine/Plugins/Media/MediaPlayerEditor/Source/MediaPlayerEditor/Private/Widgets/SMediaPlayerEditorViewport.h
+++ b/Engine/Plugins/Media/MediaPlayerEditor/Source/MediaPlayerEditor/Private/Widgets/SMediaPlayerEditorViewport.h
@@ -39,6 +39,11 @@ public:
 		UMediaTexture* InMediaTexture, const TSharedRef<ISlateStyle>& InStyle,
 		bool bInIsSoundEnabled);
 
+	/**
+	 * Enables/disables using the mouse to control the viewport.
+	 */
+	void EnableMouseControl(bool bIsEnabled) { bIsMouseControlEnabled = bIsEnabled; }
+
 public:
 
 	//~ SWidget interface
@@ -73,4 +78,7 @@ private:
 
 	/** The style set to use for this widget. */
 	TSharedPtr<ISlateStyle> Style;
+	
+	/** True if the mouse can control things. */
+	bool bIsMouseControlEnabled;
 };
\ No newline at end of file
diff --git a/Engine/Plugins/Media/MediaPlayerEditor/Source/MediaPlayerEditor/Public/Widgets/SMediaPlayerEditorViewer.h b/Engine/Plugins/Media/MediaPlayerEditor/Source/MediaPlayerEditor/Public/Widgets/SMediaPlayerEditorViewer.h
--- a/Engine/Plugins/Media/MediaPlayerEditor/Source/MediaPlayerEditor/Public/Widgets/SMediaPlayerEditorViewer.h
+++ b/Engine/Plugins/Media/MediaPlayerEditor/Source/MediaPlayerEditor/Public/Widgets/SMediaPlayerEditorViewer.h
@@ -14,6 +14,7 @@
 class FMenuBuilder;
 class SEditableTextBox;
 class SSlider;
+class SMediaPlayerEditorViewport;
 class UMediaPlayer;
 class UMediaSoundComponent;
 class UMediaTexture;
@@ -62,6 +63,11 @@ public:
 		UMediaTexture* InMediaTexture, const TSharedRef<ISlateStyle>& InStyle,
 		bool bInIsSoundEnabled);
 
+	/**
+	 * Enables/disables using the mouse to control the viewport.
+	 */
+	void EnableMouseControl(bool bIsEnabled);
+
 public:
 
 	//~ SWidget interface
@@ -150,4 +156,7 @@ private:
 
 	/** Media URL text box. */
 	TSharedPtr<SEditableTextBox> UrlTextBox;
+
+	/** The viewport that shows the media. */
+	TSharedPtr<SMediaPlayerEditorViewport> PlayerViewport;
 };
