MediaAssets: Moved MediaSourceRenderer to editor module to fix C++ UHT issue.

#rnx

diff --git a/Engine/Plugins/Media/MediaPlayerEditor/Source/MediaPlayerEditor/Private/MediaPlayerEditorModule.cpp b/Engine/Plugins/Media/MediaPlayerEditor/Source/MediaPlayerEditor/Private/MediaPlayerEditorModule.cpp
--- a/Engine/Plugins/Media/MediaPlayerEditor/Source/MediaPlayerEditor/Private/MediaPlayerEditorModule.cpp
+++ b/Engine/Plugins/Media/MediaPlayerEditor/Source/MediaPlayerEditor/Private/MediaPlayerEditorModule.cpp
@@ -18,8 +18,10 @@
 
 #include "BaseMediaSource.h"
 #include "FileMediaSource.h"
+#include "IMediaAssetsModule.h"
 #include "MediaPlayer.h"
 #include "MediaSoundComponent.h"
+#include "MediaSourceRenderer.h"
 #include "MediaTexture.h"
 #include "PlatformMediaSource.h"
 
@@ -224,6 +226,14 @@ protected:
 	/** Registers asset thumbnail renderers .*/
 	void RegisterThumbnailRenderers()
 	{
+		IMediaAssetsModule* MediaAssetsModule = FModuleManager::LoadModulePtr<IMediaAssetsModule>("MediaAssets");
+		if (MediaAssetsModule != nullptr)
+		{
+			MediaAssetsModule->RegisterCreateMediaSourceRenderer(
+				IMediaAssetsModule::FOnCreateMediaSourceRenderer::CreateRaw(this,
+					&FMediaPlayerEditorModule::CreateMediaSourceRenderer));
+		}
+
 		UThumbnailManager::Get().RegisterCustomRenderer(UMediaTexture::StaticClass(), UTextureThumbnailRenderer::StaticClass());
 		UThumbnailManager::Get().RegisterCustomRenderer(UMediaSource::StaticClass(), UMediaSourceThumbnailRenderer::StaticClass());
 	}
@@ -235,9 +245,21 @@ protected:
 		{
 			UThumbnailManager::Get().UnregisterCustomRenderer(UMediaSource::StaticClass());
 			UThumbnailManager::Get().UnregisterCustomRenderer(UMediaTexture::StaticClass());
+
+			IMediaAssetsModule* MediaAssetsModule = FModuleManager::GetModulePtr<IMediaAssetsModule>("MediaAssets");
+			if (MediaAssetsModule != nullptr)
+			{
+				MediaAssetsModule->UnregisterCreateMediaSourceRenderer();
+			}
 		}
 	}
 
+	/** Creates an object that implements IMediaSourceRendererInterface. */
+	UObject* CreateMediaSourceRenderer()
+	{
+		return NewObject<UMediaSourceRenderer>(GetTransientPackage());
+	}
+
 protected:
 
 	/** Register all component visualizers. */
diff --git a/Engine/Plugins/Media/MediaPlayerEditor/Source/MediaPlayerEditor/Private/Shared/MediaSourceRenderer.cpp b/Engine/Plugins/Media/MediaPlayerEditor/Source/MediaPlayerEditor/Private/Shared/MediaSourceRenderer.cpp
new file mode 100644
--- /dev/null
+++ b/Engine/Plugins/Media/MediaPlayerEditor/Source/MediaPlayerEditor/Private/Shared/MediaSourceRenderer.cpp
@@ -0,0 +1,114 @@
+// Copyright Epic Games, Inc. All Rights Reserved.
+
+#include "MediaSourceRenderer.h"
+
+#include "MediaPlayer.h"
+#include "MediaSource.h"
+#include "MediaTexture.h"
+#include "UObject/Package.h"
+
+UMediaTexture* UMediaSourceRenderer::Open(UMediaSource* InMediaSource)
+{
+	if ((InMediaSource != nullptr) && (InMediaSource->Validate()))
+	{
+		// Set up the player.
+		if (MediaPlayer == nullptr)
+		{
+			MediaPlayer = NewObject<UMediaPlayer>(GetTransientPackage());
+			MediaPlayer->OnSeekCompleted.AddDynamic(this, &UMediaSourceRenderer::OnSeekCompleted);
+		}
+		else
+		{
+			MediaPlayer->Close();
+		}
+
+		// Set up the texture.
+		if (MediaTexture == nullptr)
+		{
+			MediaTexture = NewObject<UMediaTexture>(GetTransientPackage());
+			MediaTexture->NewStyleOutput = true;
+		}
+		MediaTexture->CurrentAspectRatio = 0.0f;
+		MediaTexture->SetMediaPlayer(MediaPlayer.Get());
+		MediaTexture->UpdateResource();
+		MediaSource = InMediaSource;
+
+		// Start playing the media.
+		bIsSeekActive = false;
+		FMediaPlayerOptions Options;
+		Options.PlayOnOpen = EMediaPlayerOptionBooleanOverride::Enabled;
+		Options.Loop = EMediaPlayerOptionBooleanOverride::Disabled;
+		bool bIsPlaying = MediaPlayer->OpenSourceWithOptions(InMediaSource, Options);
+		if (bIsPlaying == false)
+		{
+			Close();
+		}
+	}
+
+	return MediaTexture.Get();
+}
+
+void UMediaSourceRenderer::Tick(float DeltaTime)
+{
+	if (MediaPlayer != nullptr)
+	{
+		// Is the texture ready?
+		// The aspect ratio will change when we have something.
+		if (MediaTexture->CurrentAspectRatio != 0.0f)
+		{
+			// Send this event so the content browser can update our thumbnail.
+			if (MediaSource != nullptr)
+			{
+				FPropertyChangedEvent EmptyPropertyChangedEvent(nullptr);
+				FCoreUObjectDelegates::OnObjectPropertyChanged.Broadcast(MediaSource,
+					EmptyPropertyChangedEvent);
+			}
+			Close();
+		}
+		else
+		{
+			if (MediaPlayer->GetUrl().IsEmpty())
+			{
+				Close();
+			}
+			else if (MediaPlayer->IsPreparing() == false)
+			{
+				// Don't advance playback.
+				if (MediaPlayer->IsPlaying())
+				{
+					MediaPlayer->SetRate(0.0f);
+				}
+
+				// If we can seek, then seek.
+				if ((MediaPlayer->IsBuffering() == false) && (MediaPlayer->IsReady()))
+				{
+					if (bIsSeekActive == false)
+					{
+						bIsSeekActive = true;
+						const FTimespan MediaDuration = MediaPlayer->GetDuration();
+						FTimespan SeekTime = MediaDuration * 0.3f;
+						MediaPlayer->Seek(SeekTime);
+					}
+				}
+			}
+		}
+	}
+}
+
+void UMediaSourceRenderer::OnSeekCompleted()
+{
+	bIsSeekActive = false;
+}
+
+void UMediaSourceRenderer::Close()
+{
+	if (MediaTexture != nullptr)
+	{
+		MediaTexture->SetMediaPlayer(nullptr);
+	}
+	if (MediaPlayer != nullptr)
+	{
+		MediaPlayer->Close();
+		MediaPlayer = nullptr;
+	}
+}
diff --git a/Engine/Plugins/Media/MediaPlayerEditor/Source/MediaPlayerEditor/Private/Shared/MediaSourceRenderer.h b/Engine/Plugins/Media/MediaPlayerEditor/Source/MediaPlayerEditor/Private/Shared/MediaSourceRenderer.h
new file mode 100644
--- /dev/null
+++ b/Engine/Plugins/Media/MediaPlayerEditor/Source/MediaPlayerEditor/Private/Shared/MediaSourceRenderer.h
@@ -0,0 +1,59 @@
+// Copyright Epic Games, Inc. All Rights Reserved.
+
+#pragma once
+
+#include "MediaSourceRendererInterface.h"
+#include "TickableEditorObject.h"
+#include "UObject/Object.h"
+
+#include "MediaSourceRenderer.generated.h"
+
+class UMediaPlayer;
+class UMediaSource;
+class UMediaTexture;
+
+/** Renders a media source to a texture in editor builds. */
+UCLASS()
+class UMediaSourceRenderer : public UObject
+	, public IMediaSourceRendererInterface
+	, public FTickableEditorObject
+{
+	GENERATED_BODY()
+
+public:
+	/** IMediaSourceRendererInterface interface */
+	virtual UMediaTexture* Open(UMediaSource* InMediaSource) override;
+
+	/** FTickableEditorObject interface */
+	virtual void Tick(float DeltaTime) override;
+	virtual TStatId GetStatId() const override { RETURN_QUICK_DECLARE_CYCLE_STAT(UMediaSourceRenderer, STATGROUP_Tickables); }
+
+private:
+	/**
+	 * Callback when the media player is done seeking.
+	 */
+	UFUNCTION()
+	void OnSeekCompleted();
+
+	/**
+	 * Cleans everything up.
+	 * The media textue will remain so it can be used/reused.
+	 */
+	void Close();
+
+	/** Holds the player we are using. */
+	UPROPERTY(Transient)
+	TObjectPtr<UMediaPlayer> MediaPlayer = nullptr;
+
+	/** Holds the media source we are using. */
+	UPROPERTY(Transient)
+	TObjectPtr<UMediaSource> MediaSource = nullptr;
+
+	/** Holds the media texture we are using. */
+	UPROPERTY(Transient)
+	TObjectPtr<UMediaTexture> MediaTexture = nullptr;
+
+	/** True if we are currently seeking. */
+	bool bIsSeekActive = false;
+	
+};
diff --git a/Engine/Source/Runtime/MediaAssets/Private/Assets/MediaSource.cpp b/Engine/Source/Runtime/MediaAssets/Private/Assets/MediaSource.cpp
--- a/Engine/Source/Runtime/MediaAssets/Private/Assets/MediaSource.cpp
+++ b/Engine/Source/Runtime/MediaAssets/Private/Assets/MediaSource.cpp
@@ -2,12 +2,14 @@
 
 #include "MediaSource.h"
 
+#include "IMediaAssetsModule.h"
 #include "MediaAssetsPrivate.h"
+#include "Modules/ModuleManager.h"
 #include "Misc/Paths.h"
 #include "StreamMediaSource.h"
 
 #if WITH_EDITOR
-#include "Misc/MediaSourceRenderer.h"
+#include "MediaSourceRendererInterface.h"
 #include "UObject/Package.h"
 #endif
 
@@ -25,10 +27,21 @@ void UMediaSource::GenerateThumbnail()
 {
 	if (MediaSourceRenderer == nullptr)
 	{
-		MediaSourceRenderer = NewObject<UMediaSourceRenderer>(GetTransientPackage());
+		IMediaAssetsModule* MediaAssetsModule = FModuleManager::LoadModulePtr<IMediaAssetsModule>("MediaAssets");
+		if (MediaAssetsModule != nullptr)
+		{
+			MediaSourceRenderer = MediaAssetsModule->CreateMediaSourceRenderer();
+		}
 	}
 
-	ThumbnailImage = MediaSourceRenderer->Open(this);
+	if (MediaSourceRenderer != nullptr)
+	{
+		IMediaSourceRendererInterface* Interface = Cast<IMediaSourceRendererInterface>(MediaSourceRenderer);
+		if (Interface != nullptr)
+		{
+			ThumbnailImage = Interface->Open(this);
+		}
+	}
 }
 
 #endif // WITH_EDITOR
diff --git a/Engine/Source/Runtime/MediaAssets/Private/MediaAssetsModule.cpp b/Engine/Source/Runtime/MediaAssets/Private/MediaAssetsModule.cpp
--- a/Engine/Source/Runtime/MediaAssets/Private/MediaAssetsModule.cpp
+++ b/Engine/Source/Runtime/MediaAssets/Private/MediaAssetsModule.cpp
@@ -5,6 +5,7 @@
 #include "CoreMinimal.h"
 #include "FileMediaSource.h"
 #include "IMediaAssetsModule.h"
+#include "MediaSourceRendererInterface.h"
 #include "Misc/CoreMisc.h"
 #include "Modules/ModuleManager.h"
 #include "UObject/UObjectHash.h"
@@ -90,6 +91,26 @@ public:
 
 		return Player;
 	}
+
+	virtual void RegisterCreateMediaSourceRenderer(const FOnCreateMediaSourceRenderer& Delegate) override
+	{
+		CreateMediaSourceRendererDelegate = Delegate;
+	}
+
+	virtual void UnregisterCreateMediaSourceRenderer()
+	{
+		CreateMediaSourceRendererDelegate.Unbind();
+	}
+
+	virtual UObject* CreateMediaSourceRenderer() override
+	{
+		UObject* Renderer = nullptr;
+		if (CreateMediaSourceRendererDelegate.IsBound())
+		{
+			Renderer = CreateMediaSourceRendererDelegate.Execute();
+		}
+		return Renderer;
+	}
 	
 	//~ IModuleInterface interface
 
@@ -144,6 +165,9 @@ private:
 		TEXT("mp4"),
 		TEXT("wmv"),
 	};
+
+	/** Delegate to create an object that implements IMediaSourceRendererInterface. */
+	FOnCreateMediaSourceRenderer CreateMediaSourceRendererDelegate;
 };
 
 
diff --git a/Engine/Source/Runtime/MediaAssets/Private/Misc/MediaSourceRenderer.cpp b/Engine/Source/Runtime/MediaAssets/Private/Misc/MediaSourceRenderer.cpp
deleted file mode 100644
diff --git a/Engine/Source/Runtime/MediaAssets/Private/Misc/MediaSourceRenderer.h b/Engine/Source/Runtime/MediaAssets/Private/Misc/MediaSourceRenderer.h
deleted file mode 100644
diff --git a/Engine/Source/Runtime/MediaAssets/Public/IMediaAssetsModule.h b/Engine/Source/Runtime/MediaAssets/Public/IMediaAssetsModule.h
--- a/Engine/Source/Runtime/MediaAssets/Public/IMediaAssetsModule.h
+++ b/Engine/Source/Runtime/MediaAssets/Public/IMediaAssetsModule.h
@@ -5,6 +5,7 @@
 #include "Delegates/Delegate.h"
 #include "Modules/ModuleInterface.h"
 
+class IMediaSourceRendererInterface;
 class UMediaPlayer;
 class UObject;
 
@@ -17,6 +18,8 @@ class IMediaAssetsModule
 public:
 	/** Delegate to get a player from a UObject. */
 	DECLARE_DELEGATE_RetVal_TwoParams(UMediaPlayer*, FOnGetPlayerFromObject, UObject*, UObject*& /*PlayerProxy*/);
+	/** Delegate to create an object that implements IMediaSourceRendererInterface. */
+	DECLARE_DELEGATE_RetVal(UObject*, FOnCreateMediaSourceRenderer);
 
 	/**
 	 * Plugins should call this so they can provide a function to get a media player from an object.
@@ -45,6 +48,21 @@ public:
 	 */
 	virtual UMediaPlayer* GetPlayerFromObject(UObject* Object, UObject*& PlayerProxy) = 0;
 
+	/**
+	 * Register a delegate to create an object that implements IMediaSourceRendererInterface.
+	 */
+	virtual void RegisterCreateMediaSourceRenderer(const FOnCreateMediaSourceRenderer& Delegate) = 0;
+	
+	/**
+	 * Unregisters the delegate passed in with RegisterCreateMediaSourceRenderer.
+	 */
+	virtual void UnregisterCreateMediaSourceRenderer() = 0;
+
+	/**
+	 * Creates an object that implements IMediaSourceRendererInterface.
+	 */
+	virtual UObject* CreateMediaSourceRenderer() = 0;
+
 	/** Virtual destructor. */
 	virtual ~IMediaAssetsModule() { }
 };
diff --git a/Engine/Source/Runtime/MediaAssets/Public/MediaSource.h b/Engine/Source/Runtime/MediaAssets/Public/MediaSource.h
--- a/Engine/Source/Runtime/MediaAssets/Public/MediaSource.h
+++ b/Engine/Source/Runtime/MediaAssets/Public/MediaSource.h
@@ -21,8 +21,8 @@
 #include "MediaSource.generated.h"
 
 class FVariant;
+class IMediaSourceRendererInterface;
 class UMediaSource;
-class UMediaSourceRenderer;
 class UMediaTexture;
 struct FFrame;
 
@@ -190,7 +190,7 @@ private:
 
 	/** Renders thumnbnails for us. */
 	UPROPERTY(Transient)
-	TObjectPtr<UMediaSourceRenderer> MediaSourceRenderer = nullptr;
+	TObjectPtr<UObject> MediaSourceRenderer = nullptr;
 
 #endif
 };
diff --git a/Engine/Source/Runtime/MediaAssets/Public/MediaSourceRendererInterface.h b/Engine/Source/Runtime/MediaAssets/Public/MediaSourceRendererInterface.h
new file mode 100644
--- /dev/null
+++ b/Engine/Source/Runtime/MediaAssets/Public/MediaSourceRendererInterface.h
@@ -0,0 +1,34 @@
+// Copyright Epic Games, Inc. All Rights Reserved.
+
+#pragma once
+
+#include "UObject/Interface.h"
+
+#include "MediaSourceRendererInterface.generated.h"
+
+class UMediaSource;
+class UMediaTexture;
+
+/**
+ * Interface for an object that can render a media source to a texture.
+ */
+UINTERFACE(MinimalAPI)
+class UMediaSourceRendererInterface : public UInterface
+{
+	GENERATED_BODY()
+};
+
+class MEDIAASSETS_API IMediaSourceRendererInterface
+{
+	GENERATED_BODY()
+
+public:
+	
+	/**
+	 * Open the media source to render a texture for.
+	 *
+	 * @param	InMediaSource		Media source to play.
+	 * @return	Media texture that will hold the image.
+	 */
+	virtual UMediaTexture* Open(UMediaSource* InMediaSource) = 0;
+};
