MediaPlayerEditor: Media source actions no longer need a style passed in.

diff --git a/Engine/Plugins/Media/MediaPlayerEditor/Source/MediaPlayerEditor/Private/AssetTools/FileMediaSourceActions.cpp b/Engine/Plugins/Media/MediaPlayerEditor/Source/MediaPlayerEditor/Private/AssetTools/FileMediaSourceActions.cpp
--- a/Engine/Plugins/Media/MediaPlayerEditor/Source/MediaPlayerEditor/Private/AssetTools/FileMediaSourceActions.cpp
+++ b/Engine/Plugins/Media/MediaPlayerEditor/Source/MediaPlayerEditor/Private/AssetTools/FileMediaSourceActions.cpp
@@ -11,14 +11,6 @@
 #define LOCTEXT_NAMESPACE "AssetTypeActions"
 
 
-/* FFileMediaSourceActions constructors
- *****************************************************************************/
-
-FFileMediaSourceActions::FFileMediaSourceActions(const TSharedRef<ISlateStyle>& InStyle)
-	: Style(InStyle)
-{ }
-
-
 /* FAssetTypeActions_Base interface
  *****************************************************************************/
 
diff --git a/Engine/Plugins/Media/MediaPlayerEditor/Source/MediaPlayerEditor/Private/AssetTools/FileMediaSourceActions.h b/Engine/Plugins/Media/MediaPlayerEditor/Source/MediaPlayerEditor/Private/AssetTools/FileMediaSourceActions.h
--- a/Engine/Plugins/Media/MediaPlayerEditor/Source/MediaPlayerEditor/Private/AssetTools/FileMediaSourceActions.h
+++ b/Engine/Plugins/Media/MediaPlayerEditor/Source/MediaPlayerEditor/Private/AssetTools/FileMediaSourceActions.h
@@ -4,7 +4,6 @@
 
 #include "CoreMinimal.h"
 #include "Widgets/SWidget.h"
-#include "Styling/ISlateStyle.h"
 #include "AssetTools/MediaSourceActions.h"
 
 struct FAssetData;
@@ -15,15 +14,6 @@ struct FAssetData;
 class FFileMediaSourceActions
 	: public FMediaSourceActions
 {
-public:
-
-	/**
-	 * Creates and initializes a new instance.
-	 *
-	 * @param InStyle The style set to use for asset editor toolkits.
-	 */
-	FFileMediaSourceActions(const TSharedRef<ISlateStyle>& InStyle);
-
 public:
 
 	//~ FAssetTypeActions_Base interface
@@ -32,8 +22,4 @@ public:
 	virtual FText GetName() const override;
 	virtual UClass* GetSupportedClass() const override;
 
-private:
-
-	/** Pointer to the style set to use for toolkits. */
-	TSharedRef<ISlateStyle> Style;
 };
diff --git a/Engine/Plugins/Media/MediaPlayerEditor/Source/MediaPlayerEditor/Private/AssetTools/PlatformMediaSourceActions.cpp b/Engine/Plugins/Media/MediaPlayerEditor/Source/MediaPlayerEditor/Private/AssetTools/PlatformMediaSourceActions.cpp
--- a/Engine/Plugins/Media/MediaPlayerEditor/Source/MediaPlayerEditor/Private/AssetTools/PlatformMediaSourceActions.cpp
+++ b/Engine/Plugins/Media/MediaPlayerEditor/Source/MediaPlayerEditor/Private/AssetTools/PlatformMediaSourceActions.cpp
@@ -7,14 +7,6 @@
 #define LOCTEXT_NAMESPACE "AssetTypeActions"
 
 
-/* FPlatformMediaSourceActions constructors
- *****************************************************************************/
-
-FPlatformMediaSourceActions::FPlatformMediaSourceActions(const TSharedRef<ISlateStyle>& InStyle)
-	: Style(InStyle)
-{ }
-
-
 /* FAssetTypeActions_Base interface
  *****************************************************************************/
 
diff --git a/Engine/Plugins/Media/MediaPlayerEditor/Source/MediaPlayerEditor/Private/AssetTools/PlatformMediaSourceActions.h b/Engine/Plugins/Media/MediaPlayerEditor/Source/MediaPlayerEditor/Private/AssetTools/PlatformMediaSourceActions.h
--- a/Engine/Plugins/Media/MediaPlayerEditor/Source/MediaPlayerEditor/Private/AssetTools/PlatformMediaSourceActions.h
+++ b/Engine/Plugins/Media/MediaPlayerEditor/Source/MediaPlayerEditor/Private/AssetTools/PlatformMediaSourceActions.h
@@ -4,7 +4,6 @@
 
 #include "CoreMinimal.h"
 #include "Widgets/SWidget.h"
-#include "Styling/ISlateStyle.h"
 #include "AssetTools/MediaSourceActions.h"
 
 struct FAssetData;
@@ -15,15 +14,6 @@ struct FAssetData;
 class FPlatformMediaSourceActions
 	: public FMediaSourceActions
 {
-public:
-
-	/**
-	 * Creates and initializes a new instance.
-	 *
-	 * @param InStyle The style set to use for asset editor toolkits.
-	 */
-	FPlatformMediaSourceActions(const TSharedRef<ISlateStyle>& InStyle);
-
 public:
 
 	//~ FAssetTypeActions_Base interface
@@ -32,8 +22,4 @@ public:
 	virtual FText GetName() const override;
 	virtual UClass* GetSupportedClass() const override;
 
-private:
-
-	/** Pointer to the style set to use for toolkits. */
-	TSharedRef<ISlateStyle> Style;
 };
diff --git a/Engine/Plugins/Media/MediaPlayerEditor/Source/MediaPlayerEditor/Private/AssetTools/StreamMediaSourceActions.cpp b/Engine/Plugins/Media/MediaPlayerEditor/Source/MediaPlayerEditor/Private/AssetTools/StreamMediaSourceActions.cpp
--- a/Engine/Plugins/Media/MediaPlayerEditor/Source/MediaPlayerEditor/Private/AssetTools/StreamMediaSourceActions.cpp
+++ b/Engine/Plugins/Media/MediaPlayerEditor/Source/MediaPlayerEditor/Private/AssetTools/StreamMediaSourceActions.cpp
@@ -7,14 +7,6 @@
 #define LOCTEXT_NAMESPACE "AssetTypeActions"
 
 
-/* FStreamMediaSourceActions constructors
- *****************************************************************************/
-
-FStreamMediaSourceActions::FStreamMediaSourceActions(const TSharedRef<ISlateStyle>& InStyle)
-	: Style(InStyle)
-{ }
-
-
 /* FAssetTypeActions_Base interface
  *****************************************************************************/
 
diff --git a/Engine/Plugins/Media/MediaPlayerEditor/Source/MediaPlayerEditor/Private/AssetTools/StreamMediaSourceActions.h b/Engine/Plugins/Media/MediaPlayerEditor/Source/MediaPlayerEditor/Private/AssetTools/StreamMediaSourceActions.h
--- a/Engine/Plugins/Media/MediaPlayerEditor/Source/MediaPlayerEditor/Private/AssetTools/StreamMediaSourceActions.h
+++ b/Engine/Plugins/Media/MediaPlayerEditor/Source/MediaPlayerEditor/Private/AssetTools/StreamMediaSourceActions.h
@@ -4,7 +4,6 @@
 
 #include "CoreMinimal.h"
 #include "Widgets/SWidget.h"
-#include "Styling/ISlateStyle.h"
 #include "AssetTools/MediaSourceActions.h"
 
 struct FAssetData;
@@ -15,15 +14,6 @@ struct FAssetData;
 class FStreamMediaSourceActions
 	: public FMediaSourceActions
 {
-public:
-
-	/**
-	 * Creates and initializes a new instance.
-	 *
-	 * @param InStyle The style set to use for asset editor toolkits.
-	 */
-	FStreamMediaSourceActions(const TSharedRef<ISlateStyle>& InStyle);
-
 public:
 
 	//~ FAssetTypeActions_Base interface
@@ -32,8 +22,4 @@ public:
 	virtual FText GetName() const override;
 	virtual UClass* GetSupportedClass() const override;
 
-private:
-
-	/** Pointer to the style set to use for toolkits. */
-	TSharedRef<ISlateStyle> Style;
 };
diff --git a/Engine/Plugins/Media/MediaPlayerEditor/Source/MediaPlayerEditor/Private/MediaPlayerEditorModule.cpp b/Engine/Plugins/Media/MediaPlayerEditor/Source/MediaPlayerEditor/Private/MediaPlayerEditorModule.cpp
--- a/Engine/Plugins/Media/MediaPlayerEditor/Source/MediaPlayerEditor/Private/MediaPlayerEditorModule.cpp
+++ b/Engine/Plugins/Media/MediaPlayerEditor/Source/MediaPlayerEditor/Private/MediaPlayerEditorModule.cpp
@@ -1,8 +1,8 @@
 // Copyright Epic Games, Inc. All Rights Reserved.
 
-#include "CoreMinimal.h"
+#include "MediaPlayerEditorModule.h"
+
 #include "ComponentVisualizer.h"
-#include "Modules/ModuleInterface.h"
 #include "Modules/ModuleManager.h"
 #include "AssetToolsModule.h"
 #include "UObject/UObjectHash.h"
@@ -59,7 +59,7 @@ DEFINE_LOG_CATEGORY(LogMediaPlayerEditor);
 class FMediaPlayerEditorModule
 	: public IHasMenuExtensibility
 	, public IHasToolBarExtensibility
-	, public IModuleInterface
+	, public IMediaPlayerEditorModule
 {
 public:
 
@@ -81,6 +81,12 @@ public:
 
 public:
 
+	//~ IMediaPlayerEditorModule interface
+	virtual TSharedPtr<ISlateStyle> GetStyle()
+	{
+		return Style;
+	}
+
 	//~ IModuleInterface interface
 
 	virtual void StartupModule() override
@@ -122,13 +128,13 @@ protected:
 	{
 		IAssetTools& AssetTools = FModuleManager::LoadModuleChecked<FAssetToolsModule>("AssetTools").Get();
 
-		RegisterAssetTypeAction(AssetTools, MakeShareable(new FFileMediaSourceActions(Style.ToSharedRef())));
+		RegisterAssetTypeAction(AssetTools, MakeShareable(new FFileMediaSourceActions()));
 		RegisterAssetTypeAction(AssetTools, MakeShareable(new FMediaPlayerActions(Style.ToSharedRef())));
 		RegisterAssetTypeAction(AssetTools, MakeShareable(new FMediaPlaylistActions(Style.ToSharedRef())));
-		RegisterAssetTypeAction(AssetTools, MakeShareable(new FMediaSourceActions));
+		RegisterAssetTypeAction(AssetTools, MakeShareable(new FMediaSourceActions()));
 		RegisterAssetTypeAction(AssetTools, MakeShareable(new FMediaTextureActions));
-		RegisterAssetTypeAction(AssetTools, MakeShareable(new FPlatformMediaSourceActions(Style.ToSharedRef())));
-		RegisterAssetTypeAction(AssetTools, MakeShareable(new FStreamMediaSourceActions(Style.ToSharedRef())));
+		RegisterAssetTypeAction(AssetTools, MakeShareable(new FPlatformMediaSourceActions()));
+		RegisterAssetTypeAction(AssetTools, MakeShareable(new FStreamMediaSourceActions()));
 	}
 
 	/**
diff --git a/Engine/Plugins/Media/MediaPlayerEditor/Source/MediaPlayerEditor/Public/MediaPlayerEditorModule.h b/Engine/Plugins/Media/MediaPlayerEditor/Source/MediaPlayerEditor/Public/MediaPlayerEditorModule.h
new file mode 100644
--- /dev/null
+++ b/Engine/Plugins/Media/MediaPlayerEditor/Source/MediaPlayerEditor/Public/MediaPlayerEditorModule.h
@@ -0,0 +1,22 @@
+// Copyright Epic Games, Inc. All Rights Reserved.
+
+#pragma once
+
+#include "Modules/ModuleInterface.h"
+#include "Templates/SharedPointer.h"
+
+class FImgMediaPlayer;
+class ISlateStyle;
+
+/**
+* Interface for the MediaPlayerEditor module.
+*/
+class MEDIAPLAYEREDITOR_API IMediaPlayerEditorModule
+	: public IModuleInterface
+{
+public:
+	
+	/** Get the style used by this module. */
+	virtual TSharedPtr<ISlateStyle> GetStyle() = 0;
+
+};
