MediaPlayerEditor: Added a toolkit to preview media sources.

diff --git a/Engine/Plugins/Media/ImgMedia/ImgMedia.uplugin b/Engine/Plugins/Media/ImgMedia/ImgMedia.uplugin
--- a/Engine/Plugins/Media/ImgMedia/ImgMedia.uplugin
+++ b/Engine/Plugins/Media/ImgMedia/ImgMedia.uplugin
@@ -53,5 +53,13 @@
 			"LoadingPhase" : "PostConfigInit",
 			"PlatformAllowList" : [ "Win64" ]
 		}
+	],
+	"Plugins":
+	[
+		{
+			"Name": "MediaPlayerEditor",
+			"TargetAllowList": [ "Editor" ],
+			"Enabled": true
+		},
 	]
 }
diff --git a/Engine/Plugins/Media/ImgMedia/Source/ImgMediaEditor/ImgMediaEditor.Build.cs b/Engine/Plugins/Media/ImgMedia/Source/ImgMediaEditor/ImgMediaEditor.Build.cs
--- a/Engine/Plugins/Media/ImgMedia/Source/ImgMediaEditor/ImgMediaEditor.Build.cs
+++ b/Engine/Plugins/Media/ImgMedia/Source/ImgMediaEditor/ImgMediaEditor.Build.cs
@@ -19,6 +19,7 @@ namespace UnrealBuildTool.Rules
 					"ImgMedia",
 					"InputCore",
 					"MediaAssets",
+					"MediaPlayerEditor",
 					"RenderCore",
 					"Slate",
 					"SlateCore",
diff --git a/Engine/Plugins/Media/ImgMedia/Source/ImgMediaEditor/Private/AssetTools/ImgMediaSourceActions.cpp b/Engine/Plugins/Media/ImgMedia/Source/ImgMediaEditor/Private/AssetTools/ImgMediaSourceActions.cpp
--- a/Engine/Plugins/Media/ImgMedia/Source/ImgMediaEditor/Private/AssetTools/ImgMediaSourceActions.cpp
+++ b/Engine/Plugins/Media/ImgMedia/Source/ImgMediaEditor/Private/AssetTools/ImgMediaSourceActions.cpp
@@ -4,6 +4,8 @@
 
 #include "AssetRegistry/AssetData.h"
 #include "ImgMediaSource.h"
+#include "MediaPlayerEditorModule.h"
+#include "Toolkits/MediaSourceEditorToolkit.h"
 
 
 #define LOCTEXT_NAMESPACE "ImageMediaSourceAssetTypeActions"
@@ -55,4 +57,28 @@ FColor FImgMediaSourceActions::GetTypeColor() const
 	return FColor::White;
 }
 
+void FImgMediaSourceActions::OpenAssetEditor(const TArray<UObject*>& InObjects, TSharedPtr<IToolkitHost> EditWithinLevelEditor)
+{
+	EToolkitMode::Type Mode = EditWithinLevelEditor.IsValid()
+		? EToolkitMode::WorldCentric
+		: EToolkitMode::Standalone;
+
+	for (auto ObjIt = InObjects.CreateConstIterator(); ObjIt; ++ObjIt)
+	{
+		auto MediaSource = Cast<UMediaSource>(*ObjIt);
+
+		if (MediaSource != nullptr)
+		{
+			IMediaPlayerEditorModule* MediaPlayerEditorModule = FModuleManager::LoadModulePtr<IMediaPlayerEditorModule>("MediaPlayerEditor");
+			if (MediaPlayerEditorModule != nullptr)
+			{
+				TSharedPtr<ISlateStyle> Style = MediaPlayerEditorModule->GetStyle();
+
+				TSharedRef<FMediaSourceEditorToolkit> EditorToolkit = MakeShareable(new FMediaSourceEditorToolkit(Style.ToSharedRef()));
+				EditorToolkit->Initialize(MediaSource, Mode, EditWithinLevelEditor);
+			}
+		}
+	}
+}
+
 #undef LOCTEXT_NAMESPACE
diff --git a/Engine/Plugins/Media/ImgMedia/Source/ImgMediaEditor/Private/AssetTools/ImgMediaSourceActions.h b/Engine/Plugins/Media/ImgMedia/Source/ImgMediaEditor/Private/AssetTools/ImgMediaSourceActions.h
--- a/Engine/Plugins/Media/ImgMedia/Source/ImgMediaEditor/Private/AssetTools/ImgMediaSourceActions.h
+++ b/Engine/Plugins/Media/ImgMedia/Source/ImgMediaEditor/Private/AssetTools/ImgMediaSourceActions.h
@@ -20,4 +20,6 @@ public:
 	virtual FText GetName() const override;
 	virtual UClass* GetSupportedClass() const override;
 	virtual FColor GetTypeColor() const override;
+
+	virtual void OpenAssetEditor(const TArray<UObject*>& InObjects, TSharedPtr<IToolkitHost> EditWithinLevelEditor = TSharedPtr<IToolkitHost>()) override;
 };
diff --git a/Engine/Plugins/Media/MediaPlayerEditor/Source/MediaPlayerEditor/Private/AssetTools/MediaSourceActions.cpp b/Engine/Plugins/Media/MediaPlayerEditor/Source/MediaPlayerEditor/Private/AssetTools/MediaSourceActions.cpp
--- a/Engine/Plugins/Media/MediaPlayerEditor/Source/MediaPlayerEditor/Private/AssetTools/MediaSourceActions.cpp
+++ b/Engine/Plugins/Media/MediaPlayerEditor/Source/MediaPlayerEditor/Private/AssetTools/MediaSourceActions.cpp
@@ -2,7 +2,9 @@
 
 #include "AssetTools/MediaSourceActions.h"
 #include "AssetRegistry/AssetData.h"
+#include "MediaPlayerEditorModule.h"
 #include "MediaSource.h"
+#include "Toolkits/MediaSourceEditorToolkit.h"
 
 
 #define LOCTEXT_NAMESPACE "AssetTypeActions"
@@ -80,4 +82,29 @@ TSharedPtr<class SWidget> FMediaSourceActions::GetThumbnailOverlay(const FAssetD
 }
 
 
+void FMediaSourceActions::OpenAssetEditor(const TArray<UObject*>& InObjects, TSharedPtr<IToolkitHost> EditWithinLevelEditor)
+{
+	EToolkitMode::Type Mode = EditWithinLevelEditor.IsValid()
+		? EToolkitMode::WorldCentric
+		: EToolkitMode::Standalone;
+
+	for (auto ObjIt = InObjects.CreateConstIterator(); ObjIt; ++ObjIt)
+	{
+		auto MediaSource = Cast<UMediaSource>(*ObjIt);
+
+		if (MediaSource != nullptr)
+		{
+			IMediaPlayerEditorModule* MediaPlayerEditorModule = FModuleManager::LoadModulePtr<IMediaPlayerEditorModule>("MediaPlayerEditor");
+			if (MediaPlayerEditorModule != nullptr)
+			{
+				TSharedPtr<ISlateStyle> Style = MediaPlayerEditorModule->GetStyle();
+
+				TSharedRef<FMediaSourceEditorToolkit> EditorToolkit = MakeShareable(new FMediaSourceEditorToolkit(Style.ToSharedRef()));
+				EditorToolkit->Initialize(MediaSource, Mode, EditWithinLevelEditor);
+			}
+		}
+	}
+}
+
+
 #undef LOCTEXT_NAMESPACE
diff --git a/Engine/Plugins/Media/MediaPlayerEditor/Source/MediaPlayerEditor/Private/AssetTools/MediaSourceActions.h b/Engine/Plugins/Media/MediaPlayerEditor/Source/MediaPlayerEditor/Private/AssetTools/MediaSourceActions.h
--- a/Engine/Plugins/Media/MediaPlayerEditor/Source/MediaPlayerEditor/Private/AssetTools/MediaSourceActions.h
+++ b/Engine/Plugins/Media/MediaPlayerEditor/Source/MediaPlayerEditor/Private/AssetTools/MediaSourceActions.h
@@ -6,6 +6,7 @@
 #include "Widgets/SWidget.h"
 #include "AssetTypeActions_Base.h"
 
+class ISlateStyle; 
 struct FAssetData;
 
 /**
@@ -25,4 +26,6 @@ public:
 	virtual UClass* GetSupportedClass() const override;
 	virtual FColor GetTypeColor() const override;
 	virtual TSharedPtr<class SWidget> GetThumbnailOverlay(const FAssetData& AssetData) const override;
+
+	virtual void OpenAssetEditor(const TArray<UObject*>& InObjects, TSharedPtr<IToolkitHost> EditWithinLevelEditor = TSharedPtr<IToolkitHost>()) override;
 };
diff --git a/Engine/Plugins/Media/MediaPlayerEditor/Source/MediaPlayerEditor/Private/Toolkits/MediaSourceEditorToolkit.cpp b/Engine/Plugins/Media/MediaPlayerEditor/Source/MediaPlayerEditor/Private/Toolkits/MediaSourceEditorToolkit.cpp
new file mode 100644
--- /dev/null
+++ b/Engine/Plugins/Media/MediaPlayerEditor/Source/MediaPlayerEditor/Private/Toolkits/MediaSourceEditorToolkit.cpp
@@ -0,0 +1,387 @@
+// Copyright Epic Games, Inc. All Rights Reserved.
+
+#include "Toolkits/MediaSourceEditorToolkit.h"
+
+#include "Editor.h"
+#include "EditorReimportHandler.h"
+#include "Framework/MultiBox/MultiBoxBuilder.h"
+#include "MediaPlayer.h"
+#include "MediaPlaylist.h"
+#include "MediaSource.h"
+#include "Models/MediaPlayerEditorCommands.h"
+#include "SlateOptMacros.h"
+#include "Styling/AppStyle.h"
+#include "Widgets/Docking/SDockTab.h"
+#include "Widgets/SMediaPlayerEditorInfo.h"
+#include "Widgets/SMediaPlayerEditorMedia.h"
+#include "Widgets/SMediaPlayerEditorPlaylist.h"
+#include "Widgets/SMediaPlayerEditorStats.h"
+#include "Widgets/SMediaPlayerEditorViewer.h"
+#include "Widgets/SMediaSourceEditorDetails.h"
+#include "UObject/Package.h"
+
+
+#define LOCTEXT_NAMESPACE "FMediaSourceEditorToolkit"
+
+
+namespace MediaSourceEditorToolkit
+{
+	static const FName AppIdentifier("MediaSourceEditorApp");
+	static const FName DetailsTabId("Details");
+	static const FName InfoTabId("Info");
+	static const FName ViewerTabId("Viewer");
+}
+
+
+/* FMediaSourceEditorToolkit structors
+ *****************************************************************************/
+
+FMediaSourceEditorToolkit::FMediaSourceEditorToolkit(const TSharedRef<ISlateStyle>& InStyle)
+	: MediaPlayer(nullptr)
+	, MediaSource(nullptr)
+	, Style(InStyle)
+{ }
+
+
+FMediaSourceEditorToolkit::~FMediaSourceEditorToolkit()
+{
+	FReimportManager::Instance()->OnPreReimport().RemoveAll(this);
+	FReimportManager::Instance()->OnPostReimport().RemoveAll(this);
+
+	GEditor->UnregisterForUndo(this);
+}
+
+
+/* FMediaSourceEditorToolkit interface
+ *****************************************************************************/
+
+void FMediaSourceEditorToolkit::Initialize(UMediaSource* InMediaSource, const EToolkitMode::Type InMode, const TSharedPtr<IToolkitHost>& InToolkitHost)
+{
+	MediaSource = InMediaSource;
+
+	if (MediaSource == nullptr)
+	{
+		return;
+	}
+
+	// support undo/redo
+	MediaSource->SetFlags(RF_Transactional);
+	GEditor->RegisterForUndo(this);
+
+	MediaPlayer = NewObject<UMediaPlayer>(GetTransientPackage());
+	MediaPlayer->SetLooping(false);
+	MediaPlayer->PlayOnOpen = false;
+	MediaPlayer->OpenSource(MediaSource);
+
+	BindCommands();
+
+	// create tab layout
+	const TSharedRef<FTabManager::FLayout> Layout = FTabManager::NewLayout("Standalone_MediaSourceEditor_v0.1aaab")
+		->AddArea
+		(
+			FTabManager::NewPrimaryArea()
+				->SetOrientation(Orient_Horizontal)
+				->Split
+				(
+					// viewer
+					FTabManager::NewStack()
+						->AddTab(MediaSourceEditorToolkit::ViewerTabId, ETabState::OpenedTab)
+						->SetHideTabWell(true)
+						->SetSizeCoefficient(0.6f)
+						
+				)
+				->Split
+				(
+					FTabManager::NewSplitter()
+						->SetOrientation(Orient_Vertical)
+						->SetSizeCoefficient(0.4f)
+						->Split
+						(
+							// Media details tab.
+							FTabManager::NewStack()
+							->AddTab(MediaSourceEditorToolkit::InfoTabId, ETabState::OpenedTab)
+							->SetSizeCoefficient(0.2f)
+						)
+						->Split
+						(
+							// Details tab.
+							FTabManager::NewStack()
+								->AddTab(MediaSourceEditorToolkit::DetailsTabId, ETabState::OpenedTab)
+								->SetSizeCoefficient(0.8f)
+						)
+				)
+		);
+
+	FAssetEditorToolkit::InitAssetEditor(
+		InMode,
+		InToolkitHost,
+		MediaSourceEditorToolkit::AppIdentifier,
+		Layout,
+		true /*bCreateDefaultStandaloneMenu*/,
+		true /*bCreateDefaultToolbar*/,
+		InMediaSource
+	);
+	
+	ExtendToolBar();
+	RegenerateMenusAndToolbars();
+}
+
+
+/* FAssetEditorToolkit interface
+ *****************************************************************************/
+
+FString FMediaSourceEditorToolkit::GetDocumentationLink() const
+{
+	return FString(TEXT("WorkingWithMedia/IntegratingMedia/MediaFramework"));
+}
+
+
+void FMediaSourceEditorToolkit::OnClose()
+{
+	if (MediaPlayer != nullptr)
+	{
+		MediaPlayer->Close();
+	}
+}
+
+
+void FMediaSourceEditorToolkit::RegisterTabSpawners(const TSharedRef<FTabManager>& InTabManager)
+{
+	WorkspaceMenuCategory = InTabManager->AddLocalWorkspaceMenuCategory(LOCTEXT("WorkspaceMenu_MediaSourceEditor", "Media Source Editor"));
+	auto WorkspaceMenuCategoryRef = WorkspaceMenuCategory.ToSharedRef();
+
+	FAssetEditorToolkit::RegisterTabSpawners(InTabManager);
+
+	InTabManager->RegisterTabSpawner(MediaSourceEditorToolkit::DetailsTabId, FOnSpawnTab::CreateSP(this, &FMediaSourceEditorToolkit::HandleTabManagerSpawnTab, MediaSourceEditorToolkit::DetailsTabId))
+		.SetDisplayName(LOCTEXT("DetailsTabName", "Details"))
+		.SetGroup(WorkspaceMenuCategoryRef)
+		.SetIcon(FSlateIcon(FAppStyle::GetAppStyleSetName(), "LevelEditor.Tabs.Details"));
+
+	InTabManager->RegisterTabSpawner(MediaSourceEditorToolkit::InfoTabId, FOnSpawnTab::CreateSP(this, &FMediaSourceEditorToolkit::HandleTabManagerSpawnTab, MediaSourceEditorToolkit::InfoTabId))
+		.SetDisplayName(LOCTEXT("InfoTabName", "Info"))
+		.SetGroup(WorkspaceMenuCategoryRef)
+		.SetIcon(FSlateIcon(Style->GetStyleSetName(), "MediaPlayerEditor.Tabs.Info"));
+
+	InTabManager->RegisterTabSpawner(MediaSourceEditorToolkit::ViewerTabId, FOnSpawnTab::CreateSP(this, &FMediaSourceEditorToolkit::HandleTabManagerSpawnTab, MediaSourceEditorToolkit::ViewerTabId))
+		.SetDisplayName(LOCTEXT("PlayerTabName", "Player"))
+		.SetGroup(WorkspaceMenuCategoryRef)
+		.SetIcon(FSlateIcon(Style->GetStyleSetName(), "MediaPlayerEditor.Tabs.Player"));
+}
+
+
+void FMediaSourceEditorToolkit::UnregisterTabSpawners(const TSharedRef<class FTabManager>& InTabManager)
+{
+	FAssetEditorToolkit::UnregisterTabSpawners(InTabManager);
+
+	InTabManager->UnregisterTabSpawner(MediaSourceEditorToolkit::ViewerTabId);
+	InTabManager->UnregisterTabSpawner(MediaSourceEditorToolkit::InfoTabId);
+	InTabManager->UnregisterTabSpawner(MediaSourceEditorToolkit::DetailsTabId);
+}
+
+
+/* IToolkit interface
+ *****************************************************************************/
+
+FText FMediaSourceEditorToolkit::GetBaseToolkitName() const
+{
+	return LOCTEXT("AppLabel", "Media Source Editor");
+}
+
+
+FName FMediaSourceEditorToolkit::GetToolkitFName() const
+{
+	return FName("MediaSourceEditor");
+}
+
+
+FLinearColor FMediaSourceEditorToolkit::GetWorldCentricTabColorScale() const
+{
+	return FLinearColor(0.3f, 0.2f, 0.5f, 0.5f);
+}
+
+
+FString FMediaSourceEditorToolkit::GetWorldCentricTabPrefix() const
+{
+	return LOCTEXT("WorldCentricTabPrefix", "MediaSource ").ToString();
+}
+
+
+/* FGCObject interface
+ *****************************************************************************/
+
+void FMediaSourceEditorToolkit::AddReferencedObjects(FReferenceCollector& Collector)
+{
+	Collector.AddReferencedObject(MediaPlayer);
+	Collector.AddReferencedObject(MediaSource);
+}
+
+
+/* FEditorUndoClient interface
+*****************************************************************************/
+
+void FMediaSourceEditorToolkit::PostUndo(bool bSuccess)
+{
+	// do nothing
+}
+
+
+void FMediaSourceEditorToolkit::PostRedo(bool bSuccess)
+{
+	PostUndo(bSuccess);
+}
+
+
+/* FMediaSourceEditorToolkit implementation
+ *****************************************************************************/
+
+void FMediaSourceEditorToolkit::BindCommands()
+{
+	const FMediaPlayerEditorCommands& Commands = FMediaPlayerEditorCommands::Get();
+
+	ToolkitCommands->MapAction(
+		Commands.CloseMedia,
+		FExecuteAction::CreateLambda([this] { MediaPlayer->Close(); }),
+		FCanExecuteAction::CreateLambda([this] { return !MediaPlayer->GetUrl().IsEmpty(); })
+	);
+
+	ToolkitCommands->MapAction(
+		Commands.ForwardMedia,
+		FExecuteAction::CreateLambda([this]{ MediaPlayer->SetRate(GetForwardRate()); }),
+		FCanExecuteAction::CreateLambda([this]{ return MediaPlayer->IsReady() && MediaPlayer->SupportsRate(GetForwardRate(), false); })
+	);
+
+	ToolkitCommands->MapAction(
+		Commands.NextMedia,
+		FExecuteAction::CreateLambda([this]{ MediaPlayer->Next(); }),
+		FCanExecuteAction::CreateLambda([this]{ return (MediaPlayer->GetPlaylistRef().Num() > 1); })
+	);
+
+	ToolkitCommands->MapAction(
+		Commands.PauseMedia,
+		FExecuteAction::CreateLambda([this]{ MediaPlayer->Pause(); }),
+		FCanExecuteAction::CreateLambda([this]{ return MediaPlayer->CanPause() && !MediaPlayer->IsPaused(); })
+	);
+
+	ToolkitCommands->MapAction(
+		Commands.PlayMedia,
+		FExecuteAction::CreateLambda([this]{ MediaPlayer->Play(); }),
+		FCanExecuteAction::CreateLambda([this]{ return MediaPlayer->IsReady() && (!MediaPlayer->IsPlaying() || (MediaPlayer->GetRate() != 1.0f)); })
+	);
+
+	ToolkitCommands->MapAction(
+		Commands.PreviousMedia,
+		FExecuteAction::CreateLambda([this]{ MediaPlayer->Previous(); }),
+		FCanExecuteAction::CreateLambda([this]{ return (MediaPlayer->GetPlaylistRef().Num() > 1); })
+	);
+
+	ToolkitCommands->MapAction(
+		Commands.ReverseMedia,
+		FExecuteAction::CreateLambda([this]{ MediaPlayer->SetRate(GetReverseRate()); } ),
+		FCanExecuteAction::CreateLambda([this]{ return MediaPlayer->IsReady() && MediaPlayer->SupportsRate(GetReverseRate(), false); })
+	);
+
+	ToolkitCommands->MapAction(
+		Commands.RewindMedia,
+		FExecuteAction::CreateLambda([this]{ MediaPlayer->Rewind(); }),
+		FCanExecuteAction::CreateLambda([this]{ return MediaPlayer->IsReady() && MediaPlayer->SupportsSeeking() && MediaPlayer->GetTime() > FTimespan::Zero(); })
+	);
+}
+
+
+BEGIN_SLATE_FUNCTION_BUILD_OPTIMIZATION
+void FMediaSourceEditorToolkit::ExtendToolBar()
+{
+	struct Local
+	{
+		static void FillToolbar(FToolBarBuilder& ToolbarBuilder, const TSharedRef<FUICommandList> ToolkitCommands)
+		{
+			ToolbarBuilder.BeginSection("PlaybackControls");
+			{
+				ToolbarBuilder.AddToolBarButton(FMediaPlayerEditorCommands::Get().PreviousMedia);
+				ToolbarBuilder.AddToolBarButton(FMediaPlayerEditorCommands::Get().RewindMedia);
+				ToolbarBuilder.AddToolBarButton(FMediaPlayerEditorCommands::Get().ReverseMedia);
+				ToolbarBuilder.AddToolBarButton(FMediaPlayerEditorCommands::Get().PlayMedia);
+				ToolbarBuilder.AddToolBarButton(FMediaPlayerEditorCommands::Get().PauseMedia);
+				ToolbarBuilder.AddToolBarButton(FMediaPlayerEditorCommands::Get().ForwardMedia);
+				ToolbarBuilder.AddToolBarButton(FMediaPlayerEditorCommands::Get().NextMedia);
+			}
+			ToolbarBuilder.EndSection();
+
+			ToolbarBuilder.BeginSection("MediaControls");
+			{
+				ToolbarBuilder.AddToolBarButton(FMediaPlayerEditorCommands::Get().CloseMedia);
+			}
+			ToolbarBuilder.EndSection();
+		}
+	};
+
+
+	TSharedPtr<FExtender> ToolbarExtender = MakeShareable(new FExtender);
+
+	ToolbarExtender->AddToolBarExtension(
+		"Asset",
+		EExtensionHook::After,
+		GetToolkitCommands(),
+		FToolBarExtensionDelegate::CreateStatic(&Local::FillToolbar, GetToolkitCommands())
+	);
+
+	AddToolbarExtender(ToolbarExtender);
+}
+END_SLATE_FUNCTION_BUILD_OPTIMIZATION
+
+
+float FMediaSourceEditorToolkit::GetForwardRate() const
+{
+	float Rate = MediaPlayer->GetRate();
+
+	if (Rate < 1.0f)
+	{
+		Rate = 1.0f;
+	}
+
+	return 2.0f * Rate;
+}
+
+
+float FMediaSourceEditorToolkit::GetReverseRate() const
+{
+	float Rate = MediaPlayer->GetRate();
+
+	if (Rate > -1.0f)
+	{
+		return -1.0f;
+	}
+
+	return 2.0f * Rate;
+}
+
+
+/* FMediaSourceEditorToolkit callbacks
+ *****************************************************************************/
+
+TSharedRef<SDockTab> FMediaSourceEditorToolkit::HandleTabManagerSpawnTab(const FSpawnTabArgs& Args, FName TabIdentifier)
+{
+	TSharedPtr<SWidget> TabWidget = SNullWidget::NullWidget;
+
+	if (TabIdentifier == MediaSourceEditorToolkit::DetailsTabId)
+	{
+		TabWidget = SNew(SMediaSourceEditorDetails, *MediaSource, Style);
+	}
+	else if (TabIdentifier == MediaSourceEditorToolkit::InfoTabId)
+	{
+		TabWidget = SNew(SMediaPlayerEditorInfo, *MediaPlayer, Style);
+	}
+	else if (TabIdentifier == MediaSourceEditorToolkit::ViewerTabId)
+	{
+		TabWidget = SNew(SMediaPlayerEditorViewer, *MediaPlayer, nullptr, Style, true);
+	}
+
+	return SNew(SDockTab)
+		.TabRole(ETabRole::PanelTab)
+		[
+			TabWidget.ToSharedRef()
+		];
+}
+
+
+#undef LOCTEXT_NAMESPACE
diff --git a/Engine/Plugins/Media/MediaPlayerEditor/Source/MediaPlayerEditor/Private/Widgets/SMediaSourceEditorDetails.cpp b/Engine/Plugins/Media/MediaPlayerEditor/Source/MediaPlayerEditor/Private/Widgets/SMediaSourceEditorDetails.cpp
new file mode 100644
--- /dev/null
+++ b/Engine/Plugins/Media/MediaPlayerEditor/Source/MediaPlayerEditor/Private/Widgets/SMediaSourceEditorDetails.cpp
@@ -0,0 +1,38 @@
+// Copyright Epic Games, Inc. All Rights Reserved.
+
+#include "Widgets/SMediaSourceEditorDetails.h"
+
+#include "IDetailsView.h"
+#include "MediaSource.h"
+#include "Modules/ModuleManager.h"
+#include "PropertyEditorModule.h"
+
+#define LOCTEXT_NAMESPACE "SMediaSourceEditorDetails"
+
+void SMediaSourceEditorDetails::Construct(const FArguments& InArgs, UMediaSource& InMediaSource, const TSharedRef<ISlateStyle>& InStyle)
+{
+	MediaSource = &InMediaSource;
+
+	// Initialize details view
+	FDetailsViewArgs DetailsViewArgs;
+	{
+		DetailsViewArgs.bAllowSearch = true;
+		DetailsViewArgs.bHideSelectionTip = true;
+		DetailsViewArgs.bLockable = false;
+		DetailsViewArgs.bSearchInitialKeyFocus = true;
+		DetailsViewArgs.bUpdatesFromSelection = false;
+		DetailsViewArgs.bShowOptions = true;
+		DetailsViewArgs.bShowModifiedPropertiesOption = false;
+	}
+
+	TSharedPtr<IDetailsView> DetailsView = FModuleManager::GetModuleChecked<FPropertyEditorModule>("PropertyEditor").CreateDetailView(DetailsViewArgs);
+	DetailsView->SetObject(MediaSource);
+
+	ChildSlot
+	[
+		DetailsView.ToSharedRef()
+	];
+}
+
+
+#undef LOCTEXT_NAMESPACE
diff --git a/Engine/Plugins/Media/MediaPlayerEditor/Source/MediaPlayerEditor/Private/Widgets/SMediaSourceEditorDetails.h b/Engine/Plugins/Media/MediaPlayerEditor/Source/MediaPlayerEditor/Private/Widgets/SMediaSourceEditorDetails.h
new file mode 100644
--- /dev/null
+++ b/Engine/Plugins/Media/MediaPlayerEditor/Source/MediaPlayerEditor/Private/Widgets/SMediaSourceEditorDetails.h
@@ -0,0 +1,37 @@
+// Copyright Epic Games, Inc. All Rights Reserved.
+
+#pragma once
+
+#include "Styling/ISlateStyle.h"
+#include "Widgets/DeclarativeSyntaxSupport.h"
+#include "Widgets/SCompoundWidget.h"
+
+class UMediaSource;
+
+/**
+ * Implements the details panel of the MediaSource asset editor.
+ */
+class SMediaSourceEditorDetails
+	: public SCompoundWidget
+{
+public:
+
+	SLATE_BEGIN_ARGS(SMediaSourceEditorDetails) { }
+	SLATE_END_ARGS()
+
+public:
+
+	/**
+	 * Construct this widget.
+	 *
+	 * @param InArgs		The declaration data for this widget.
+	 * @param InMediaSource	The MediaSource to show the details for.
+	 * @param InStyleSet	The style set to use.
+	 */
+	void Construct(const FArguments& InArgs, UMediaSource& InMediaSource, const TSharedRef<ISlateStyle>& InStyle);
+
+private:
+
+	/** Pointer to the MediaPlate that is being viewed. */
+	TObjectPtr<UMediaSource> MediaSource;
+};
diff --git a/Engine/Plugins/Media/MediaPlayerEditor/Source/MediaPlayerEditor/Public/Toolkits/MediaSourceEditorToolkit.h b/Engine/Plugins/Media/MediaPlayerEditor/Source/MediaPlayerEditor/Public/Toolkits/MediaSourceEditorToolkit.h
new file mode 100644
--- /dev/null
+++ b/Engine/Plugins/Media/MediaPlayerEditor/Source/MediaPlayerEditor/Public/Toolkits/MediaSourceEditorToolkit.h
@@ -0,0 +1,117 @@
+// Copyright Epic Games, Inc. All Rights Reserved.
+
+#pragma once
+
+#include "EditorUndoClient.h"
+#include "Framework/Docking/TabManager.h"
+#include "Styling/ISlateStyle.h"
+#include "Toolkits/AssetEditorToolkit.h"
+#include "UObject/GCObject.h"
+
+class UMediaPlayer;
+class UMediaSource;
+
+/**
+ * Implements an Editor toolkit for media sources.
+ */
+class MEDIAPLAYEREDITOR_API FMediaSourceEditorToolkit
+	: public FAssetEditorToolkit
+	, public FEditorUndoClient
+	, public FGCObject
+{
+public:
+
+	/**
+	 * Creates and initializes a new instance.
+	 *
+	 * @param InStyle The style set to use.
+	 */
+	FMediaSourceEditorToolkit(const TSharedRef<ISlateStyle>& InStyle);
+
+	/** Virtual destructor. */
+	virtual ~FMediaSourceEditorToolkit();
+
+public:
+
+	/**
+	 * Initializes the editor tool kit.
+	 *
+	 * @param InMediaSource The UMediaSource asset to edit.
+	 * @param InMode The mode to create the toolkit in.
+	 * @param InToolkitHost The toolkit host.
+	 */
+	void Initialize(UMediaSource* InMediaSource, const EToolkitMode::Type InMode, const TSharedPtr<IToolkitHost>& InToolkitHost);
+
+public:
+
+	//~ FAssetEditorToolkit interface
+
+	virtual FString GetDocumentationLink() const override;
+	virtual void OnClose() override;
+	virtual void RegisterTabSpawners(const TSharedRef<FTabManager>& TabManager) override;
+	virtual void UnregisterTabSpawners(const TSharedRef<FTabManager>& TabManager) override;
+
+public:
+
+	//~ IToolkit interface
+
+	virtual FText GetBaseToolkitName() const override;
+	virtual FName GetToolkitFName() const override;
+	virtual FLinearColor GetWorldCentricTabColorScale() const override;
+	virtual FString GetWorldCentricTabPrefix() const override;
+
+public:
+
+	//~ FGCObject interface
+
+	virtual void AddReferencedObjects(FReferenceCollector& Collector) override;
+	virtual FString GetReferencerName() const override
+	{
+		return TEXT("FMediaSourceEditorToolkit");
+	}
+
+protected:
+
+	//~ FEditorUndoClient interface
+
+	virtual void PostUndo(bool bSuccess) override;
+	virtual void PostRedo(bool bSuccess) override;
+
+protected:
+
+	/** Binds the UI commands to delegates. */
+	void BindCommands();
+
+	/** Builds the toolbar widget for the media player editor. */
+	void ExtendToolBar();
+
+	/**
+	 * Gets the playback rate for fast forward.
+	 *
+	 * @return Forward playback rate.
+	 */
+	float GetForwardRate() const;
+
+	/**
+	 * Gets the playback rate for reverse.
+	 *
+	 * @return Reverse playback rate.
+	 */
+	float GetReverseRate() const;
+
+private:
+
+	/** Callback for spawning tabs. */
+	TSharedRef<SDockTab> HandleTabManagerSpawnTab(const FSpawnTabArgs& Args, FName TabIdentifier);
+
+private:
+
+	/** The media player to play the media with. */
+	TObjectPtr<UMediaPlayer> MediaPlayer;
+
+	/** The media source asset being edited. */
+	TObjectPtr<UMediaSource> MediaSource;
+
+	/** Pointer to the style set to use for toolkits. */
+	TSharedRef<ISlateStyle> Style;
+};
