MediaPlayerEditor: Added open button to the media source toolkit.

diff --git a/Engine/Plugins/Media/MediaPlayerEditor/Source/MediaPlayerEditor/Private/Models/MediaPlayerEditorCommands.cpp b/Engine/Plugins/Media/MediaPlayerEditor/Source/MediaPlayerEditor/Private/Models/MediaPlayerEditorCommands.cpp
--- a/Engine/Plugins/Media/MediaPlayerEditor/Source/MediaPlayerEditor/Private/Models/MediaPlayerEditorCommands.cpp
+++ b/Engine/Plugins/Media/MediaPlayerEditor/Source/MediaPlayerEditor/Private/Models/MediaPlayerEditorCommands.cpp
@@ -11,6 +11,7 @@ void FMediaPlayerEditorCommands::RegisterCommands()
 	UI_COMMAND(CloseMedia, "Close", "Close the currently opened media", EUserInterfaceActionType::Button, FInputChord());
 	UI_COMMAND(ForwardMedia, "Forward", "Fast forward media playback", EUserInterfaceActionType::Button, FInputChord());
 	UI_COMMAND(NextMedia, "Next", "Jump to next item in the play list", EUserInterfaceActionType::Button, FInputChord());
+	UI_COMMAND(OpenMedia, "Open", "Open the current media", EUserInterfaceActionType::Button, FInputChord());
 	UI_COMMAND(PauseMedia, "Pause", "Pause media playback", EUserInterfaceActionType::Button, FInputChord());
 	UI_COMMAND(PlayMedia, "Play", "Start media playback", EUserInterfaceActionType::Button, FInputChord());
 	UI_COMMAND(PreviousMedia, "Prev", "Jump to previous item in the play list", EUserInterfaceActionType::Button, FInputChord());
diff --git a/Engine/Plugins/Media/MediaPlayerEditor/Source/MediaPlayerEditor/Private/Models/MediaPlayerEditorCommands.h b/Engine/Plugins/Media/MediaPlayerEditor/Source/MediaPlayerEditor/Private/Models/MediaPlayerEditorCommands.h
--- a/Engine/Plugins/Media/MediaPlayerEditor/Source/MediaPlayerEditor/Private/Models/MediaPlayerEditorCommands.h
+++ b/Engine/Plugins/Media/MediaPlayerEditor/Source/MediaPlayerEditor/Private/Models/MediaPlayerEditorCommands.h
@@ -35,6 +35,9 @@ public:
 	/** Jump to next item in the play list. */
 	TSharedPtr<FUICommandInfo> NextMedia;
 
+	/** Open the current media. */
+	TSharedPtr<FUICommandInfo> OpenMedia;
+
 	/** Pauses media playback. */
 	TSharedPtr<FUICommandInfo> PauseMedia;
 
diff --git a/Engine/Plugins/Media/MediaPlayerEditor/Source/MediaPlayerEditor/Private/Shared/MediaPlayerEditorStyle.cpp b/Engine/Plugins/Media/MediaPlayerEditor/Source/MediaPlayerEditor/Private/Shared/MediaPlayerEditorStyle.cpp
--- a/Engine/Plugins/Media/MediaPlayerEditor/Source/MediaPlayerEditor/Private/Shared/MediaPlayerEditorStyle.cpp
+++ b/Engine/Plugins/Media/MediaPlayerEditor/Source/MediaPlayerEditor/Private/Shared/MediaPlayerEditorStyle.cpp
@@ -56,6 +56,8 @@ FMediaPlayerEditorStyle::FMediaPlayerEditorStyle()
 	Set("MediaPlayerEditor.ForwardMedia.Small", new IMAGE_BRUSH("icon_forward_40x", Icon20x20));
 	Set("MediaPlayerEditor.NextMedia", new IMAGE_BRUSH("icon_step_40x", Icon40x40));
 	Set("MediaPlayerEditor.NextMedia.Small", new IMAGE_BRUSH("icon_step_40x", Icon20x20));
+	Set("MediaPlayerEditor.OpenMedia", new IMAGE_BRUSH("icon_open_40x", Icon40x40));
+	Set("MediaPlayerEditor.OpenMedia.Small", new IMAGE_BRUSH("icon_open_40x", Icon20x20));
 	Set("MediaPlayerEditor.PauseMedia", new IMAGE_BRUSH("icon_pause_40x", Icon40x40));
 	Set("MediaPlayerEditor.PauseMedia.Small", new IMAGE_BRUSH("icon_pause_40x", Icon20x20));
 	Set("MediaPlayerEditor.PlayMedia", new IMAGE_BRUSH("icon_play_40x", Icon40x40));
diff --git a/Engine/Plugins/Media/MediaPlayerEditor/Source/MediaPlayerEditor/Private/Toolkits/MediaSourceEditorToolkit.cpp b/Engine/Plugins/Media/MediaPlayerEditor/Source/MediaPlayerEditor/Private/Toolkits/MediaSourceEditorToolkit.cpp
--- a/Engine/Plugins/Media/MediaPlayerEditor/Source/MediaPlayerEditor/Private/Toolkits/MediaSourceEditorToolkit.cpp
+++ b/Engine/Plugins/Media/MediaPlayerEditor/Source/MediaPlayerEditor/Private/Toolkits/MediaSourceEditorToolkit.cpp
@@ -70,8 +70,7 @@ void FMediaSourceEditorToolkit::Initialize(UMediaSource* InMediaSource, const ET
 
 	MediaPlayer = NewObject<UMediaPlayer>(GetTransientPackage());
 	MediaPlayer->SetLooping(false);
-	MediaPlayer->PlayOnOpen = false;
-	MediaPlayer->OpenSource(MediaSource);
+	MediaPlayer->PlayOnOpen = true;
 
 	BindCommands();
 
@@ -256,6 +255,12 @@ void FMediaSourceEditorToolkit::BindCommands()
 		FCanExecuteAction::CreateLambda([this]{ return (MediaPlayer->GetPlaylistRef().Num() > 1); })
 	);
 
+	ToolkitCommands->MapAction(
+		Commands.OpenMedia,
+		FExecuteAction::CreateLambda([this] { MediaPlayer->OpenSource(MediaSource); }),
+		FCanExecuteAction::CreateLambda([this] { return true; })
+	);
+
 	ToolkitCommands->MapAction(
 		Commands.PauseMedia,
 		FExecuteAction::CreateLambda([this]{ MediaPlayer->Pause(); }),
@@ -309,6 +314,7 @@ void FMediaSourceEditorToolkit::ExtendToolBar()
 
 			ToolbarBuilder.BeginSection("MediaControls");
 			{
+				ToolbarBuilder.AddToolBarButton(FMediaPlayerEditorCommands::Get().OpenMedia);
 				ToolbarBuilder.AddToolBarButton(FMediaPlayerEditorCommands::Get().CloseMedia);
 			}
 			ToolbarBuilder.EndSection();
