MediaPlayerEditor: Removed url from media source toolkit as its not needed.

diff --git a/Engine/Plugins/Media/MediaPlayerEditor/Source/MediaPlayerEditor/Private/Toolkits/MediaSourceEditorToolkit.cpp b/Engine/Plugins/Media/MediaPlayerEditor/Source/MediaPlayerEditor/Private/Toolkits/MediaSourceEditorToolkit.cpp
--- a/Engine/Plugins/Media/MediaPlayerEditor/Source/MediaPlayerEditor/Private/Toolkits/MediaSourceEditorToolkit.cpp
+++ b/Engine/Plugins/Media/MediaPlayerEditor/Source/MediaPlayerEditor/Private/Toolkits/MediaSourceEditorToolkit.cpp
@@ -379,7 +379,8 @@ TSharedRef<SDockTab> FMediaSourceEditorToolkit::HandleTabManagerSpawnTab(const F
 	}
 	else if (TabIdentifier == MediaSourceEditorToolkit::ViewerTabId)
 	{
-		TabWidget = SNew(SMediaPlayerEditorViewer, *MediaPlayer, nullptr, Style, true);
+		TabWidget = SNew(SMediaPlayerEditorViewer, *MediaPlayer, nullptr, Style, true)
+			.bShowUrl(false);
 	}
 
 	return SNew(SDockTab)
