MediaPlayerEditor: Added player details tab to media source toolkit.

diff --git a/Engine/Plugins/Media/MediaPlayerEditor/Source/MediaPlayerEditor/Private/Toolkits/MediaSourceEditorToolkit.cpp b/Engine/Plugins/Media/MediaPlayerEditor/Source/MediaPlayerEditor/Private/Toolkits/MediaSourceEditorToolkit.cpp
--- a/Engine/Plugins/Media/MediaPlayerEditor/Source/MediaPlayerEditor/Private/Toolkits/MediaSourceEditorToolkit.cpp
+++ b/Engine/Plugins/Media/MediaPlayerEditor/Source/MediaPlayerEditor/Private/Toolkits/MediaSourceEditorToolkit.cpp
@@ -12,6 +12,7 @@
 #include "SlateOptMacros.h"
 #include "Styling/AppStyle.h"
 #include "Widgets/Docking/SDockTab.h"
+#include "Widgets/SMediaPlayerEditorDetails.h"
 #include "Widgets/SMediaPlayerEditorInfo.h"
 #include "Widgets/SMediaPlayerEditorMedia.h"
 #include "Widgets/SMediaPlayerEditorPlaylist.h"
@@ -29,6 +30,7 @@ namespace MediaSourceEditorToolkit
 	static const FName AppIdentifier("MediaSourceEditorApp");
 	static const FName DetailsTabId("Details");
 	static const FName InfoTabId("Info");
+	static const FName PlayerDetailsTabId("PlayerDetails");
 	static const FName ViewerTabId("Viewer");
 }
 
@@ -75,7 +77,7 @@ void FMediaSourceEditorToolkit::Initialize(UMediaSource* InMediaSource, const ET
 	BindCommands();
 
 	// create tab layout
-	const TSharedRef<FTabManager::FLayout> Layout = FTabManager::NewLayout("Standalone_MediaSourceEditor_v0.1aaab")
+	const TSharedRef<FTabManager::FLayout> Layout = FTabManager::NewLayout("Standalone_MediaSourceEditor_v0.2")
 		->AddArea
 		(
 			FTabManager::NewPrimaryArea()
@@ -106,6 +108,8 @@ void FMediaSourceEditorToolkit::Initialize(UMediaSource* InMediaSource, const ET
 							// Details tab.
 							FTabManager::NewStack()
 								->AddTab(MediaSourceEditorToolkit::DetailsTabId, ETabState::OpenedTab)
+								->AddTab(MediaSourceEditorToolkit::PlayerDetailsTabId, ETabState::OpenedTab)
+								->SetForegroundTab(MediaSourceEditorToolkit::DetailsTabId)
 								->SetSizeCoefficient(0.8f)
 						)
 				)
@@ -161,6 +165,11 @@ void FMediaSourceEditorToolkit::RegisterTabSpawners(const TSharedRef<FTabManager
 		.SetGroup(WorkspaceMenuCategoryRef)
 		.SetIcon(FSlateIcon(Style->GetStyleSetName(), "MediaPlayerEditor.Tabs.Info"));
 
+	InTabManager->RegisterTabSpawner(MediaSourceEditorToolkit::PlayerDetailsTabId, FOnSpawnTab::CreateSP(this, &FMediaSourceEditorToolkit::HandleTabManagerSpawnTab, MediaSourceEditorToolkit::PlayerDetailsTabId))
+		.SetDisplayName(LOCTEXT("PlayerDetailsTabName", "Player Details"))
+		.SetGroup(WorkspaceMenuCategoryRef)
+		.SetIcon(FSlateIcon(FAppStyle::GetAppStyleSetName(), "LevelEditor.Tabs.Details"));
+
 	InTabManager->RegisterTabSpawner(MediaSourceEditorToolkit::ViewerTabId, FOnSpawnTab::CreateSP(this, &FMediaSourceEditorToolkit::HandleTabManagerSpawnTab, MediaSourceEditorToolkit::ViewerTabId))
 		.SetDisplayName(LOCTEXT("PlayerTabName", "Player"))
 		.SetGroup(WorkspaceMenuCategoryRef)
@@ -173,6 +182,7 @@ void FMediaSourceEditorToolkit::UnregisterTabSpawners(const TSharedRef<class FTa
 	FAssetEditorToolkit::UnregisterTabSpawners(InTabManager);
 
 	InTabManager->UnregisterTabSpawner(MediaSourceEditorToolkit::ViewerTabId);
+	InTabManager->UnregisterTabSpawner(MediaSourceEditorToolkit::PlayerDetailsTabId);
 	InTabManager->UnregisterTabSpawner(MediaSourceEditorToolkit::InfoTabId);
 	InTabManager->UnregisterTabSpawner(MediaSourceEditorToolkit::DetailsTabId);
 }
@@ -377,6 +387,10 @@ TSharedRef<SDockTab> FMediaSourceEditorToolkit::HandleTabManagerSpawnTab(const F
 	{
 		TabWidget = SNew(SMediaPlayerEditorInfo, *MediaPlayer, Style);
 	}
+	else if (TabIdentifier == MediaSourceEditorToolkit::PlayerDetailsTabId)
+	{
+		TabWidget = SNew(SMediaPlayerEditorDetails, *MediaPlayer, Style);
+	}
 	else if (TabIdentifier == MediaSourceEditorToolkit::ViewerTabId)
 	{
 		TabWidget = SNew(SMediaPlayerEditorViewer, *MediaPlayer, nullptr, Style, true)
