MediaPlayerEditor: Cleaned up a bit to make MediaSourceActions public instead of MediaSourceEditorToolkit.h.

#rnx

diff --git a/Engine/Plugins/Media/ImgMedia/Source/ImgMediaEditor/Private/AssetTools/ImgMediaSourceActions.cpp b/Engine/Plugins/Media/ImgMedia/Source/ImgMediaEditor/Private/AssetTools/ImgMediaSourceActions.cpp
--- a/Engine/Plugins/Media/ImgMedia/Source/ImgMediaEditor/Private/AssetTools/ImgMediaSourceActions.cpp
+++ b/Engine/Plugins/Media/ImgMedia/Source/ImgMediaEditor/Private/AssetTools/ImgMediaSourceActions.cpp
@@ -2,11 +2,7 @@
 
 #include "AssetTools/ImgMediaSourceActions.h"
 
-#include "AssetRegistry/AssetData.h"
 #include "ImgMediaSource.h"
-#include "MediaPlayerEditorModule.h"
-#include "Toolkits/MediaSourceEditorToolkit.h"
-
 
 #define LOCTEXT_NAMESPACE "ImageMediaSourceAssetTypeActions"
 
@@ -15,33 +11,6 @@ bool FImgMediaSourceActions::CanFilter()
 	return true;
 }
 
-FText FImgMediaSourceActions::GetAssetDescription(const struct FAssetData& AssetData) const
-{
-	const UImgMediaSource* ImgMediaSource = Cast<UImgMediaSource>(AssetData.GetAsset());
-
-	if (ImgMediaSource != nullptr)
-	{
-		const FString Url = ImgMediaSource->GetUrl();
-
-		if (Url.IsEmpty())
-		{
-			return LOCTEXT("ImgAssetTypeActions_URLMissing", "Warning: Missing URL detected!");
-		}
-
-		if (!ImgMediaSource->Validate())
-		{
-			return LOCTEXT("AssetTypeActions_ImgMediaSourceInvalid", "Warning: Invalid settings detected!");
-		}
-	}
-
-	return FText::GetEmpty();
-}
-
-uint32 FImgMediaSourceActions::GetCategories()
-{
-	return EAssetTypeCategories::Media;
-}
-
 FText FImgMediaSourceActions::GetName() const
 {
 	return LOCTEXT("AssetTypeActions_ImgMediaSource", "Img Media Source");
@@ -52,33 +21,4 @@ UClass* FImgMediaSourceActions::GetSupportedClass() const
 	return UImgMediaSource::StaticClass();
 }
 
-FColor FImgMediaSourceActions::GetTypeColor() const
-{
-	return FColor::White;
-}
-
-void FImgMediaSourceActions::OpenAssetEditor(const TArray<UObject*>& InObjects, TSharedPtr<IToolkitHost> EditWithinLevelEditor)
-{
-	EToolkitMode::Type Mode = EditWithinLevelEditor.IsValid()
-		? EToolkitMode::WorldCentric
-		: EToolkitMode::Standalone;
-
-	for (auto ObjIt = InObjects.CreateConstIterator(); ObjIt; ++ObjIt)
-	{
-		auto MediaSource = Cast<UMediaSource>(*ObjIt);
-
-		if (MediaSource != nullptr)
-		{
-			IMediaPlayerEditorModule* MediaPlayerEditorModule = FModuleManager::LoadModulePtr<IMediaPlayerEditorModule>("MediaPlayerEditor");
-			if (MediaPlayerEditorModule != nullptr)
-			{
-				TSharedPtr<ISlateStyle> Style = MediaPlayerEditorModule->GetStyle();
-
-				TSharedRef<FMediaSourceEditorToolkit> EditorToolkit = MakeShareable(new FMediaSourceEditorToolkit(Style.ToSharedRef()));
-				EditorToolkit->Initialize(MediaSource, Mode, EditWithinLevelEditor);
-			}
-		}
-	}
-}
-
 #undef LOCTEXT_NAMESPACE
diff --git a/Engine/Plugins/Media/ImgMedia/Source/ImgMediaEditor/Private/AssetTools/ImgMediaSourceActions.h b/Engine/Plugins/Media/ImgMedia/Source/ImgMediaEditor/Private/AssetTools/ImgMediaSourceActions.h
--- a/Engine/Plugins/Media/ImgMedia/Source/ImgMediaEditor/Private/AssetTools/ImgMediaSourceActions.h
+++ b/Engine/Plugins/Media/ImgMedia/Source/ImgMediaEditor/Private/AssetTools/ImgMediaSourceActions.h
@@ -3,23 +3,18 @@
 #pragma once
 
 #include "CoreMinimal.h"
-#include "AssetTypeActions_Base.h"
+#include "AssetTools/MediaSourceActions.h"
 
 
 /**
  * Implements an action for ImgMediaSource assets.
  */
-class FImgMediaSourceActions : public FAssetTypeActions_Base
+class FImgMediaSourceActions : public FMediaSourceActions
 {
 public:
 
 	//~ FAssetTypeActions_Base interface
 	virtual bool CanFilter() override;
-	virtual FText GetAssetDescription(const FAssetData& AssetData) const override;
-	virtual uint32 GetCategories() override;
 	virtual FText GetName() const override;
 	virtual UClass* GetSupportedClass() const override;
-	virtual FColor GetTypeColor() const override;
-
-	virtual void OpenAssetEditor(const TArray<UObject*>& InObjects, TSharedPtr<IToolkitHost> EditWithinLevelEditor = TSharedPtr<IToolkitHost>()) override;
 };
diff --git a/Engine/Plugins/Media/MediaPlayerEditor/Source/MediaPlayerEditor/Private/AssetTools/MediaSourceActions.h b/Engine/Plugins/Media/MediaPlayerEditor/Source/MediaPlayerEditor/Private/AssetTools/MediaSourceActions.h
deleted file mode 100644
diff --git a/Engine/Plugins/Media/MediaPlayerEditor/Source/MediaPlayerEditor/Private/Toolkits/MediaSourceEditorToolkit.h b/Engine/Plugins/Media/MediaPlayerEditor/Source/MediaPlayerEditor/Private/Toolkits/MediaSourceEditorToolkit.h
new file mode 100644
--- /dev/null
+++ b/Engine/Plugins/Media/MediaPlayerEditor/Source/MediaPlayerEditor/Private/Toolkits/MediaSourceEditorToolkit.h
@@ -0,0 +1,117 @@
+// Copyright Epic Games, Inc. All Rights Reserved.
+
+#pragma once
+
+#include "EditorUndoClient.h"
+#include "Framework/Docking/TabManager.h"
+#include "Styling/ISlateStyle.h"
+#include "Toolkits/AssetEditorToolkit.h"
+#include "UObject/GCObject.h"
+
+class UMediaPlayer;
+class UMediaSource;
+
+/**
+ * Implements an Editor toolkit for media sources.
+ */
+class MEDIAPLAYEREDITOR_API FMediaSourceEditorToolkit
+	: public FAssetEditorToolkit
+	, public FEditorUndoClient
+	, public FGCObject
+{
+public:
+
+	/**
+	 * Creates and initializes a new instance.
+	 *
+	 * @param InStyle The style set to use.
+	 */
+	FMediaSourceEditorToolkit(const TSharedRef<ISlateStyle>& InStyle);
+
+	/** Virtual destructor. */
+	virtual ~FMediaSourceEditorToolkit();
+
+public:
+
+	/**
+	 * Initializes the editor tool kit.
+	 *
+	 * @param InMediaSource The UMediaSource asset to edit.
+	 * @param InMode The mode to create the toolkit in.
+	 * @param InToolkitHost The toolkit host.
+	 */
+	void Initialize(UMediaSource* InMediaSource, const EToolkitMode::Type InMode, const TSharedPtr<IToolkitHost>& InToolkitHost);
+
+public:
+
+	//~ FAssetEditorToolkit interface
+
+	virtual FString GetDocumentationLink() const override;
+	virtual void OnClose() override;
+	virtual void RegisterTabSpawners(const TSharedRef<FTabManager>& TabManager) override;
+	virtual void UnregisterTabSpawners(const TSharedRef<FTabManager>& TabManager) override;
+
+public:
+
+	//~ IToolkit interface
+
+	virtual FText GetBaseToolkitName() const override;
+	virtual FName GetToolkitFName() const override;
+	virtual FLinearColor GetWorldCentricTabColorScale() const override;
+	virtual FString GetWorldCentricTabPrefix() const override;
+
+public:
+
+	//~ FGCObject interface
+
+	virtual void AddReferencedObjects(FReferenceCollector& Collector) override;
+	virtual FString GetReferencerName() const override
+	{
+		return TEXT("FMediaSourceEditorToolkit");
+	}
+
+protected:
+
+	//~ FEditorUndoClient interface
+
+	virtual void PostUndo(bool bSuccess) override;
+	virtual void PostRedo(bool bSuccess) override;
+
+protected:
+
+	/** Binds the UI commands to delegates. */
+	void BindCommands();
+
+	/** Builds the toolbar widget for the media player editor. */
+	void ExtendToolBar();
+
+	/**
+	 * Gets the playback rate for fast forward.
+	 *
+	 * @return Forward playback rate.
+	 */
+	float GetForwardRate() const;
+
+	/**
+	 * Gets the playback rate for reverse.
+	 *
+	 * @return Reverse playback rate.
+	 */
+	float GetReverseRate() const;
+
+private:
+
+	/** Callback for spawning tabs. */
+	TSharedRef<SDockTab> HandleTabManagerSpawnTab(const FSpawnTabArgs& Args, FName TabIdentifier);
+
+private:
+
+	/** The media player to play the media with. */
+	TObjectPtr<UMediaPlayer> MediaPlayer;
+
+	/** The media source asset being edited. */
+	TObjectPtr<UMediaSource> MediaSource;
+
+	/** Pointer to the style set to use for toolkits. */
+	TSharedRef<ISlateStyle> Style;
+};
diff --git a/Engine/Plugins/Media/MediaPlayerEditor/Source/MediaPlayerEditor/Public/AssetTools/MediaSourceActions.h b/Engine/Plugins/Media/MediaPlayerEditor/Source/MediaPlayerEditor/Public/AssetTools/MediaSourceActions.h
new file mode 100644
--- /dev/null
+++ b/Engine/Plugins/Media/MediaPlayerEditor/Source/MediaPlayerEditor/Public/AssetTools/MediaSourceActions.h
@@ -0,0 +1,31 @@
+// Copyright Epic Games, Inc. All Rights Reserved.
+
+#pragma once
+
+#include "CoreMinimal.h"
+#include "Widgets/SWidget.h"
+#include "AssetTypeActions_Base.h"
+
+class ISlateStyle; 
+struct FAssetData;
+
+/**
+ * Implements an action for UMediaSource assets.
+ */
+class MEDIAPLAYEREDITOR_API FMediaSourceActions
+	: public FAssetTypeActions_Base
+{
+public:
+
+	//~ FAssetTypeActions_Base interface
+
+	virtual bool CanFilter() override;
+	virtual FText GetAssetDescription(const FAssetData& AssetData) const override;
+	virtual uint32 GetCategories() override;
+	virtual FText GetName() const override;
+	virtual UClass* GetSupportedClass() const override;
+	virtual FColor GetTypeColor() const override;
+	virtual TSharedPtr<class SWidget> GetThumbnailOverlay(const FAssetData& AssetData) const override;
+
+	virtual void OpenAssetEditor(const TArray<UObject*>& InObjects, TSharedPtr<IToolkitHost> EditWithinLevelEditor = TSharedPtr<IToolkitHost>()) override;
+};
diff --git a/Engine/Plugins/Media/MediaPlayerEditor/Source/MediaPlayerEditor/Public/Toolkits/MediaSourceEditorToolkit.h b/Engine/Plugins/Media/MediaPlayerEditor/Source/MediaPlayerEditor/Public/Toolkits/MediaSourceEditorToolkit.h
deleted file mode 100644
