MediaPlayerEditor: Added media details tab to the media source toolkit.

diff --git a/Engine/Plugins/Media/MediaPlayerEditor/Source/MediaPlayerEditor/Private/Toolkits/MediaSourceEditorToolkit.cpp b/Engine/Plugins/Media/MediaPlayerEditor/Source/MediaPlayerEditor/Private/Toolkits/MediaSourceEditorToolkit.cpp
--- a/Engine/Plugins/Media/MediaPlayerEditor/Source/MediaPlayerEditor/Private/Toolkits/MediaSourceEditorToolkit.cpp
+++ b/Engine/Plugins/Media/MediaPlayerEditor/Source/MediaPlayerEditor/Private/Toolkits/MediaSourceEditorToolkit.cpp
@@ -8,13 +8,13 @@
 #include "MediaPlayer.h"
 #include "MediaPlaylist.h"
 #include "MediaSource.h"
+#include "MediaTexture.h"
 #include "Models/MediaPlayerEditorCommands.h"
 #include "SlateOptMacros.h"
 #include "Styling/AppStyle.h"
 #include "Widgets/Docking/SDockTab.h"
 #include "Widgets/SMediaPlayerEditorDetails.h"
-#include "Widgets/SMediaPlayerEditorInfo.h"
-#include "Widgets/SMediaPlayerEditorMedia.h"
+#include "Widgets/SMediaPlayerEditorMediaDetails.h"
 #include "Widgets/SMediaPlayerEditorPlaylist.h"
 #include "Widgets/SMediaPlayerEditorStats.h"
 #include "Widgets/SMediaPlayerEditorViewer.h"
@@ -29,7 +29,7 @@ namespace MediaSourceEditorToolkit
 {
 	static const FName AppIdentifier("MediaSourceEditorApp");
 	static const FName DetailsTabId("Details");
-	static const FName InfoTabId("Info");
+	static const FName MediaDetailsTabId("MediaDetails");
 	static const FName PlayerDetailsTabId("PlayerDetails");
 	static const FName ViewerTabId("Viewer");
 }
@@ -41,6 +41,7 @@ namespace MediaSourceEditorToolkit
 FMediaSourceEditorToolkit::FMediaSourceEditorToolkit(const TSharedRef<ISlateStyle>& InStyle)
 	: MediaPlayer(nullptr)
 	, MediaSource(nullptr)
+	, MediaTexture(nullptr)
 	, Style(InStyle)
 { }
 
@@ -74,10 +75,18 @@ void FMediaSourceEditorToolkit::Initialize(UMediaSource* InMediaSource, const ET
 	MediaPlayer->SetLooping(false);
 	MediaPlayer->PlayOnOpen = true;
 
+	MediaTexture = NewObject<UMediaTexture>(GetTransientPackage(), NAME_None, RF_Transient | RF_Public);
+	if (MediaTexture != nullptr)
+	{
+		MediaTexture->AutoClear = true;
+		MediaTexture->SetMediaPlayer(MediaPlayer);
+		MediaTexture->UpdateResource();
+	}
+
 	BindCommands();
 
 	// create tab layout
-	const TSharedRef<FTabManager::FLayout> Layout = FTabManager::NewLayout("Standalone_MediaSourceEditor_v0.2")
+	const TSharedRef<FTabManager::FLayout> Layout = FTabManager::NewLayout("Standalone_MediaSourceEditor_v0.3")
 		->AddArea
 		(
 			FTabManager::NewPrimaryArea()
@@ -100,7 +109,7 @@ void FMediaSourceEditorToolkit::Initialize(UMediaSource* InMediaSource, const ET
 						(
 							// Media details tab.
 							FTabManager::NewStack()
-							->AddTab(MediaSourceEditorToolkit::InfoTabId, ETabState::OpenedTab)
+							->AddTab(MediaSourceEditorToolkit::MediaDetailsTabId, ETabState::OpenedTab)
 							->SetSizeCoefficient(0.2f)
 						)
 						->Split
@@ -160,8 +169,8 @@ void FMediaSourceEditorToolkit::RegisterTabSpawners(const TSharedRef<FTabManager
 		.SetGroup(WorkspaceMenuCategoryRef)
 		.SetIcon(FSlateIcon(FAppStyle::GetAppStyleSetName(), "LevelEditor.Tabs.Details"));
 
-	InTabManager->RegisterTabSpawner(MediaSourceEditorToolkit::InfoTabId, FOnSpawnTab::CreateSP(this, &FMediaSourceEditorToolkit::HandleTabManagerSpawnTab, MediaSourceEditorToolkit::InfoTabId))
-		.SetDisplayName(LOCTEXT("InfoTabName", "Info"))
+	InTabManager->RegisterTabSpawner(MediaSourceEditorToolkit::MediaDetailsTabId, FOnSpawnTab::CreateSP(this, &FMediaSourceEditorToolkit::HandleTabManagerSpawnTab, MediaSourceEditorToolkit::MediaDetailsTabId))
+		.SetDisplayName(LOCTEXT("MediaDetailsTabName", "Media Details"))
 		.SetGroup(WorkspaceMenuCategoryRef)
 		.SetIcon(FSlateIcon(Style->GetStyleSetName(), "MediaPlayerEditor.Tabs.Info"));
 
@@ -183,7 +192,7 @@ void FMediaSourceEditorToolkit::UnregisterTabSpawners(const TSharedRef<class FTa
 
 	InTabManager->UnregisterTabSpawner(MediaSourceEditorToolkit::ViewerTabId);
 	InTabManager->UnregisterTabSpawner(MediaSourceEditorToolkit::PlayerDetailsTabId);
-	InTabManager->UnregisterTabSpawner(MediaSourceEditorToolkit::InfoTabId);
+	InTabManager->UnregisterTabSpawner(MediaSourceEditorToolkit::MediaDetailsTabId);
 	InTabManager->UnregisterTabSpawner(MediaSourceEditorToolkit::DetailsTabId);
 }
 
@@ -222,6 +231,7 @@ void FMediaSourceEditorToolkit::AddReferencedObjects(FReferenceCollector& Collec
 {
 	Collector.AddReferencedObject(MediaPlayer);
 	Collector.AddReferencedObject(MediaSource);
+	Collector.AddReferencedObject(MediaTexture);
 }
 
 
@@ -383,9 +393,9 @@ TSharedRef<SDockTab> FMediaSourceEditorToolkit::HandleTabManagerSpawnTab(const F
 	{
 		TabWidget = SNew(SMediaSourceEditorDetails, *MediaSource, Style);
 	}
-	else if (TabIdentifier == MediaSourceEditorToolkit::InfoTabId)
+	else if (TabIdentifier == MediaSourceEditorToolkit::MediaDetailsTabId)
 	{
-		TabWidget = SNew(SMediaPlayerEditorInfo, *MediaPlayer, Style);
+		TabWidget = SNew(SMediaPlayerEditorMediaDetails, MediaPlayer, MediaTexture);
 	}
 	else if (TabIdentifier == MediaSourceEditorToolkit::PlayerDetailsTabId)
 	{
@@ -393,7 +403,7 @@ TSharedRef<SDockTab> FMediaSourceEditorToolkit::HandleTabManagerSpawnTab(const F
 	}
 	else if (TabIdentifier == MediaSourceEditorToolkit::ViewerTabId)
 	{
-		TabWidget = SNew(SMediaPlayerEditorViewer, *MediaPlayer, nullptr, Style, true)
+		TabWidget = SNew(SMediaPlayerEditorViewer, *MediaPlayer, MediaTexture, Style, true)
 			.bShowUrl(false);
 	}
 
diff --git a/Engine/Plugins/Media/MediaPlayerEditor/Source/MediaPlayerEditor/Private/Toolkits/MediaSourceEditorToolkit.h b/Engine/Plugins/Media/MediaPlayerEditor/Source/MediaPlayerEditor/Private/Toolkits/MediaSourceEditorToolkit.h
--- a/Engine/Plugins/Media/MediaPlayerEditor/Source/MediaPlayerEditor/Private/Toolkits/MediaSourceEditorToolkit.h
+++ b/Engine/Plugins/Media/MediaPlayerEditor/Source/MediaPlayerEditor/Private/Toolkits/MediaSourceEditorToolkit.h
@@ -10,6 +10,7 @@
 
 class UMediaPlayer;
 class UMediaSource;
+class UMediaTexture;
 
 /**
  * Implements an Editor toolkit for media sources.
@@ -112,6 +113,9 @@ private:
 	/** The media source asset being edited. */
 	TObjectPtr<UMediaSource> MediaSource;
 
+	/** The media texture to output the media to. */
+	TObjectPtr<UMediaTexture> MediaTexture;
+
 	/** Pointer to the style set to use for toolkits. */
 	TSharedRef<ISlateStyle> Style;
 };
diff --git a/Engine/Plugins/Media/MediaPlayerEditor/Source/MediaPlayerEditor/Private/Widgets/SMediaPlayerEditorMediaDetails.cpp b/Engine/Plugins/Media/MediaPlayerEditor/Source/MediaPlayerEditor/Private/Widgets/SMediaPlayerEditorMediaDetails.cpp
new file mode 100644
--- /dev/null
+++ b/Engine/Plugins/Media/MediaPlayerEditor/Source/MediaPlayerEditor/Private/Widgets/SMediaPlayerEditorMediaDetails.cpp
@@ -0,0 +1,185 @@
+// Copyright Epic Games, Inc. All Rights Reserved.
+
+#include "Widgets/SMediaPlayerEditorMediaDetails.h"
+#include "IDetailsView.h"
+#include "MediaPlayer.h"
+#include "MediaTexture.h"
+#include "Modules/ModuleManager.h"
+#include "PropertyEditorModule.h"
+#include "Widgets/Layout/SScrollBox.h"
+#include "Widgets/Text/STextBlock.h"
+
+#define LOCTEXT_NAMESPACE "SMediaPlayerEditorMediaDetails"
+
+/* SMediaPlayerEditorMediaDetails interface
+ *****************************************************************************/
+
+void SMediaPlayerEditorMediaDetails::Construct(const FArguments& InArgs,
+	UMediaPlayer* InMediaPlayer, UMediaTexture* InMediaTexture)
+{
+	MediaPlayer = InMediaPlayer;
+	MediaTexture = InMediaTexture;
+
+	ChildSlot
+		[
+			SNew(SScrollBox)
+
+			// Add details.
+			+ SScrollBox::Slot()
+				[
+					SNew(SHorizontalBox)
+
+					// Left side.
+					+ SHorizontalBox::Slot()
+						[
+							SNew(SVerticalBox)
+
+							// Resolution.
+							+ SVerticalBox::Slot()
+								.AutoHeight()
+								.VAlign(VAlign_Center)
+								.Padding(4.0f)
+								[
+									SAssignNew(ResolutionText, STextBlock)
+								]
+
+							// Frame rate.
+							+ SVerticalBox::Slot()
+								.AutoHeight()
+								.VAlign(VAlign_Center)
+								.Padding(4.0f)
+								[
+									SAssignNew(FrameRateText, STextBlock)
+								]
+
+							// Resource size.
+							+ SVerticalBox::Slot()
+								.AutoHeight()
+								.VAlign(VAlign_Center)
+								.Padding(4.0f)
+								[
+									SAssignNew(ResourceSizeText, STextBlock)
+								]
+
+							// Method.
+							+ SVerticalBox::Slot()
+								.AutoHeight()
+								.VAlign(VAlign_Center)
+								.Padding(4.0f)
+								[
+									SAssignNew(MethodText, STextBlock)
+								]
+						]
+
+					// Right side.
+					+ SHorizontalBox::Slot()
+						[
+							SNew(SVerticalBox)
+
+							// Format.
+							+ SVerticalBox::Slot()
+								.AutoHeight()
+								.VAlign(VAlign_Center)
+								.Padding(4.0f)
+								[
+									SAssignNew(FormatText, STextBlock)
+								]
+
+							// LOD bias.
+							+ SVerticalBox::Slot()
+								.AutoHeight()
+								.VAlign(VAlign_Center)
+								.Padding(4.0f)
+								[
+									SAssignNew(LODBiasText, STextBlock)
+								]
+
+							// Num mips.
+							+ SVerticalBox::Slot()
+								.AutoHeight()
+								.VAlign(VAlign_Center)
+								.Padding(4.0f)
+								[
+									SAssignNew(NumMipsText, STextBlock)
+								]
+
+							// Num tiles.
+							+ SVerticalBox::Slot()
+								.AutoHeight()
+								.VAlign(VAlign_Center)
+								.Padding(4.0f)
+								[
+									SAssignNew(NumTilesText, STextBlock)
+								]
+						]
+				]
+
+		];
+
+	UpdateDetails();
+}
+
+
+void SMediaPlayerEditorMediaDetails::Tick(const FGeometry& AllottedGeometry, const double InCurrentTime, const float InDeltaTime)
+{
+	// Call parent.
+	SCompoundWidget::Tick(AllottedGeometry, InCurrentTime, InDeltaTime);
+
+	UpdateDetails();
+}
+
+void SMediaPlayerEditorMediaDetails::UpdateDetails()
+{
+	FString Format;
+	float FrameRate = 0.0f;
+	int32 LODBias = 0;
+	FText Method;
+	int32 NumMips = 0;
+	int32 NumTotalTiles = 0;
+	int64 ResourceSize = 0;
+	int32 SurfaceWidth = 0;
+	int32 SurfaceHeight = 0;
+
+	// Get player info.
+	if (MediaPlayer != nullptr)
+	{
+		FrameRate = MediaPlayer->GetVideoTrackFrameRate(INDEX_NONE, INDEX_NONE);
+		Format = MediaPlayer->GetVideoTrackType(INDEX_NONE, INDEX_NONE);
+		FIntPoint NumTiles(EForceInit::ForceInitToZero);
+		MediaPlayer->GetMediaInfo<FIntPoint>(NumTiles, UMediaPlayer::MediaInfoNameSourceNumTiles.Resolve());
+		NumTotalTiles = NumTiles.X * NumTiles.Y;
+	}
+
+	// Get texture info,
+	if (MediaTexture != nullptr)
+	{
+		LODBias = MediaTexture->GetCachedLODBias();
+		Method = MediaTexture->IsCurrentlyVirtualTextured() ?
+			LOCTEXT("MethodVirtualStreamed", "Virtual Streamed")
+			: (!MediaTexture->IsStreamable() ? LOCTEXT("QuickInfo_MethodNotStreamed", "Not Streamed")
+				: LOCTEXT("MethodStreamed", "Streamed"));
+		NumMips = MediaTexture->GetTextureNumMips();
+		ResourceSize = (MediaTexture->GetResourceSizeBytes(EResourceSizeMode::Exclusive) + 512) / 1024;
+		SurfaceWidth = MediaTexture->GetSurfaceWidth();
+		SurfaceHeight = MediaTexture->GetSurfaceHeight();
+	}
+
+	// Update text.
+	FormatText->SetText(FText::Format(LOCTEXT("Format", "Format: {0}"),
+		FText::FromString(Format)));
+	FrameRateText->SetText(FText::Format(LOCTEXT("FrameRate", "Frame Rate: {0}"), 
+		FText::AsNumber(FrameRate)));
+	LODBiasText->SetText(FText::Format(LOCTEXT("LODBias", "Combined LOD Bias: {0}"),
+		FText::AsNumber(LODBias)));
+	MethodText->SetText(FText::Format(LOCTEXT("Method", "Method: {0}"), Method));
+	NumMipsText->SetText(FText::Format(LOCTEXT("NumberOfMips", "Number Of Mips: {0}"),
+		FText::AsNumber(NumMips)));
+	NumTilesText->SetText(FText::Format(LOCTEXT("NumberOfTiles", "Number Of Tiles: {0}"),
+		FText::AsNumber(NumTotalTiles)));
+	ResolutionText->SetText(FText::Format(LOCTEXT("Resolution", "Resolution: {0}x{1}"),
+		FText::AsNumber(SurfaceWidth), FText::AsNumber(SurfaceHeight)));
+	ResourceSizeText->SetText(FText::Format(LOCTEXT("ResourceSize", "Resource Size: {0} KB"),
+		FText::AsNumber(ResourceSize)));
+}
+
+#undef LOCTEXT_NAMESPACE
diff --git a/Engine/Plugins/Media/MediaPlayerEditor/Source/MediaPlayerEditor/Public/Widgets/SMediaPlayerEditorMediaDetails.h b/Engine/Plugins/Media/MediaPlayerEditor/Source/MediaPlayerEditor/Public/Widgets/SMediaPlayerEditorMediaDetails.h
new file mode 100644
--- /dev/null
+++ b/Engine/Plugins/Media/MediaPlayerEditor/Source/MediaPlayerEditor/Public/Widgets/SMediaPlayerEditorMediaDetails.h
@@ -0,0 +1,60 @@
+// Copyright Epic Games, Inc. All Rights Reserved.
+
+#pragma once
+
+#include "Styling/ISlateStyle.h"
+#include "Widgets/DeclarativeSyntaxSupport.h"
+#include "Widgets/SCompoundWidget.h"
+
+class STextBlock;
+class UMediaPlayer;
+class UMediaTexture;;
+
+/**
+ * Implements the media details panel of the MediaPlayer asset editor.
+ */
+class SMediaPlayerEditorMediaDetails
+	: public SCompoundWidget
+{
+public:
+
+	SLATE_BEGIN_ARGS(SMediaPlayerEditorMediaDetails) { }
+	SLATE_END_ARGS()
+
+public:
+
+	/**
+	 * Construct this widget.
+	 *
+	 * @param InArgs			The declaration data for this widget.
+	 * @param InMediaPlayer		The MediaPlayer to show the details for.
+	 * @param InMediaTexture	The MediaTexture to show the details for.
+	 * 
+	 */
+	void Construct(const FArguments& InArgs, UMediaPlayer* InMediaPlate, UMediaTexture* InMediaTexture);
+
+	//~ SWidget interface
+	virtual void Tick(const FGeometry& AllottedGeometry, const double InCurrentTime, const float InDeltaTime) override;
+
+private:
+
+	/**
+	 * Updates our widgets to reflect the current state.
+	 */
+	void UpdateDetails();
+
+	/** Pointer to the MediaPlayer that is being viewed. */
+	TWeakObjectPtr<UMediaPlayer> MediaPlayer;
+	/** Pointer to the MediaTexture that is being viewed. */
+	TWeakObjectPtr<UMediaTexture> MediaTexture;
+
+	/** Our widgets. */
+	TSharedPtr<STextBlock> FormatText;
+	TSharedPtr<STextBlock> FrameRateText;
+	TSharedPtr<STextBlock> LODBiasText;
+	TSharedPtr<STextBlock> MethodText;
+	TSharedPtr<STextBlock> NumMipsText;
+	TSharedPtr<STextBlock> NumTilesText;
+	TSharedPtr<STextBlock> ResolutionText;
+	TSharedPtr<STextBlock> ResourceSizeText;
+};
