MediaPlayerEditor: Checked for OuterObjects being null as a blind fix.

diff --git a/Engine/Plugins/Media/MediaPlayerEditor/Source/MediaPlayerEditor/Private/Customizations/BaseMediaSourceCustomization.cpp b/Engine/Plugins/Media/MediaPlayerEditor/Source/MediaPlayerEditor/Private/Customizations/BaseMediaSourceCustomization.cpp
--- a/Engine/Plugins/Media/MediaPlayerEditor/Source/MediaPlayerEditor/Private/Customizations/BaseMediaSourceCustomization.cpp
+++ b/Engine/Plugins/Media/MediaPlayerEditor/Source/MediaPlayerEditor/Private/Customizations/BaseMediaSourceCustomization.cpp
@@ -225,7 +225,7 @@ FText FBaseMediaSourceCustomization::HandlePlatformPlayersComboButtonText(FStrin
 		PlatformPlayerNamesProperty->GetOuterObjects(OuterObjects);
 	}
 
-	if (OuterObjects.Num() == 0)
+	if ((OuterObjects.Num() == 0) || (OuterObjects[0] == nullptr))
 	{
 		return FText::GetEmpty();
 	}
@@ -234,7 +234,8 @@ FText FBaseMediaSourceCustomization::HandlePlatformPlayersComboButtonText(FStrin
 
 	for (int32 ObjectIndex = 1; ObjectIndex < OuterObjects.Num(); ++ObjectIndex)
 	{
-		if (Cast<UBaseMediaSource>(OuterObjects[ObjectIndex])->PlatformPlayerNames.FindRef(PlatformName) != PlayerName)
+		if ((OuterObjects[ObjectIndex] != nullptr) &&
+			(Cast<UBaseMediaSource>(OuterObjects[ObjectIndex])->PlatformPlayerNames.FindRef(PlatformName) != PlayerName))
 		{
 			return NSLOCTEXT("PropertyEditor", "MultipleValues", "Multiple Values");
 		}
