MediaPlayerEditor: Added generate thumbnail button to media source toolkit.

diff --git a/Engine/Plugins/Media/MediaPlayerEditor/Source/MediaPlayerEditor/MediaPlayerEditor.Build.cs b/Engine/Plugins/Media/MediaPlayerEditor/Source/MediaPlayerEditor/MediaPlayerEditor.Build.cs
--- a/Engine/Plugins/Media/MediaPlayerEditor/Source/MediaPlayerEditor/MediaPlayerEditor.Build.cs
+++ b/Engine/Plugins/Media/MediaPlayerEditor/Source/MediaPlayerEditor/MediaPlayerEditor.Build.cs
@@ -1,5 +1,6 @@
 // Copyright Epic Games, Inc. All Rights Reserved.
 
+using System.IO;
 using UnrealBuildTool;
 
 public class MediaPlayerEditor : ModuleRules
@@ -49,6 +50,7 @@ public class MediaPlayerEditor : ModuleRules
 				"MediaUtils",
 				"PropertyEditor",
 				"RenderCore",
+				"Renderer",
 				"RHI",
 				"TextureEditor",
 				"EditorFramework",
@@ -63,5 +65,10 @@ public class MediaPlayerEditor : ModuleRules
 				"UnrealEd",
 				"WorkspaceMenuStructure",
 			});
+
+
+		var EngineDir = Path.GetFullPath(Target.RelativeEnginePath);
+		PrivateIncludePaths.Add(Path.Combine(EngineDir, "Source", "Runtime", "Renderer", "Private"));
+
 	}
 }
diff --git a/Engine/Plugins/Media/MediaPlayerEditor/Source/MediaPlayerEditor/Private/Models/MediaPlayerEditorCommands.cpp b/Engine/Plugins/Media/MediaPlayerEditor/Source/MediaPlayerEditor/Private/Models/MediaPlayerEditorCommands.cpp
--- a/Engine/Plugins/Media/MediaPlayerEditor/Source/MediaPlayerEditor/Private/Models/MediaPlayerEditorCommands.cpp
+++ b/Engine/Plugins/Media/MediaPlayerEditor/Source/MediaPlayerEditor/Private/Models/MediaPlayerEditorCommands.cpp
@@ -10,6 +10,7 @@ void FMediaPlayerEditorCommands::RegisterCommands()
 {
 	UI_COMMAND(CloseMedia, "Close", "Close the currently opened media", EUserInterfaceActionType::Button, FInputChord());
 	UI_COMMAND(ForwardMedia, "Forward", "Fast forward media playback", EUserInterfaceActionType::Button, FInputChord());
+	UI_COMMAND(GenerateThumbnail, "Thumbnail", "Generate a thumbnail", EUserInterfaceActionType::Button, FInputChord());
 	UI_COMMAND(NextMedia, "Next", "Jump to next item in the play list", EUserInterfaceActionType::Button, FInputChord());
 	UI_COMMAND(OpenMedia, "Open", "Open the current media", EUserInterfaceActionType::Button, FInputChord());
 	UI_COMMAND(PauseMedia, "Pause", "Pause media playback", EUserInterfaceActionType::Button, FInputChord());
diff --git a/Engine/Plugins/Media/MediaPlayerEditor/Source/MediaPlayerEditor/Private/Models/MediaPlayerEditorCommands.h b/Engine/Plugins/Media/MediaPlayerEditor/Source/MediaPlayerEditor/Private/Models/MediaPlayerEditorCommands.h
--- a/Engine/Plugins/Media/MediaPlayerEditor/Source/MediaPlayerEditor/Private/Models/MediaPlayerEditorCommands.h
+++ b/Engine/Plugins/Media/MediaPlayerEditor/Source/MediaPlayerEditor/Private/Models/MediaPlayerEditorCommands.h
@@ -32,6 +32,9 @@ public:
 	/** Fast forwards media playback. */
 	TSharedPtr<FUICommandInfo> ForwardMedia;
 
+	/** Generate a thumnbnail. */
+	TSharedPtr<FUICommandInfo> GenerateThumbnail;
+
 	/** Jump to next item in the play list. */
 	TSharedPtr<FUICommandInfo> NextMedia;
 
diff --git a/Engine/Plugins/Media/MediaPlayerEditor/Source/MediaPlayerEditor/Private/Shared/MediaPlayerEditorStyle.cpp b/Engine/Plugins/Media/MediaPlayerEditor/Source/MediaPlayerEditor/Private/Shared/MediaPlayerEditorStyle.cpp
--- a/Engine/Plugins/Media/MediaPlayerEditor/Source/MediaPlayerEditor/Private/Shared/MediaPlayerEditorStyle.cpp
+++ b/Engine/Plugins/Media/MediaPlayerEditor/Source/MediaPlayerEditor/Private/Shared/MediaPlayerEditorStyle.cpp
@@ -8,17 +8,11 @@
 #include "Fonts/SlateFontInfo.h"
 #include "Math/Vector2D.h"
 #include "Misc/Paths.h"
+#include "Styling/SlateStyleMacros.h"
 #include "Styling/SlateStyleRegistry.h"
 #include "Styling/SlateTypes.h"
 #include "Styling/CoreStyle.h"
 
-
-#define IMAGE_BRUSH(RelativePath, ...) FSlateImageBrush(RootToContentDir(RelativePath, TEXT(".png")), __VA_ARGS__)
-#define BOX_BRUSH(RelativePath, ...) FSlateBoxBrush(RootToContentDir(RelativePath, TEXT(".png")), __VA_ARGS__)
-#define BORDER_BRUSH(RelativePath, ...) FSlateBorderBrush(RootToContentDir(RelativePath, TEXT(".png")), __VA_ARGS__)
-#define DEFAULT_FONT(...) FCoreStyle::GetDefaultFontStyle(__VA_ARGS__)
-
-
 /* FMediaPlayerEditorStyle structors
  *****************************************************************************/
 
@@ -32,6 +26,7 @@ FMediaPlayerEditorStyle::FMediaPlayerEditorStyle()
 	const FVector2D Icon40x40(40.0f, 40.0f);
 
 	SetContentRoot(FPaths::EnginePluginsDir() / TEXT("Media/MediaPlayerEditor/Content"));
+	SetCoreContentRoot(FPaths::EngineContentDir() / TEXT("Editor/Slate"));
 
 	// buttons
 	Set("MediaPlayerEditor.SourceButton", new IMAGE_BRUSH("btn_source_12x", Icon12x12));
@@ -54,6 +49,7 @@ FMediaPlayerEditorStyle::FMediaPlayerEditorStyle()
 	Set("MediaPlayerEditor.CloseMedia", new IMAGE_BRUSH("icon_eject_40x", Icon40x40));
 	Set("MediaPlayerEditor.ForwardMedia", new IMAGE_BRUSH("icon_forward_40x", Icon40x40));
 	Set("MediaPlayerEditor.ForwardMedia.Small", new IMAGE_BRUSH("icon_forward_40x", Icon20x20));
+	Set("MediaPlayerEditor.GenerateThumbnail", new CORE_IMAGE_BRUSH_SVG("Starship/AssetEditors/SaveThumbnail", Icon40x40));
 	Set("MediaPlayerEditor.NextMedia", new IMAGE_BRUSH("icon_step_40x", Icon40x40));
 	Set("MediaPlayerEditor.NextMedia.Small", new IMAGE_BRUSH("icon_step_40x", Icon20x20));
 	Set("MediaPlayerEditor.OpenMedia", new IMAGE_BRUSH("icon_open_40x", Icon40x40));
diff --git a/Engine/Plugins/Media/MediaPlayerEditor/Source/MediaPlayerEditor/Private/Shared/MediaSourceThumbnailRenderer.cpp b/Engine/Plugins/Media/MediaPlayerEditor/Source/MediaPlayerEditor/Private/Shared/MediaSourceThumbnailRenderer.cpp
--- a/Engine/Plugins/Media/MediaPlayerEditor/Source/MediaPlayerEditor/Private/Shared/MediaSourceThumbnailRenderer.cpp
+++ b/Engine/Plugins/Media/MediaPlayerEditor/Source/MediaPlayerEditor/Private/Shared/MediaSourceThumbnailRenderer.cpp
@@ -13,7 +13,7 @@ bool UMediaSourceThumbnailRenderer::CanVisualizeAsset(UObject* Object)
 
 void UMediaSourceThumbnailRenderer::GetThumbnailSize(UObject* Object, float Zoom, uint32& OutWidth, uint32& OutHeight) const
 {
-	UMediaTexture* ObjectTexture = GetThumbnailTextureFromObject(Object);
+	UTexture* ObjectTexture = GetThumbnailTextureFromObject(Object);
 	if (ObjectTexture != nullptr)
 	{
 		OutWidth = Zoom * ObjectTexture->GetSurfaceWidth();
@@ -28,19 +28,19 @@ void UMediaSourceThumbnailRenderer::GetThumbnailSize(UObject* Object, float Zoom
 
 void UMediaSourceThumbnailRenderer::Draw(UObject* Object, int32 X, int32 Y, uint32 Width, uint32 Height, FRenderTarget* RenderTarget, FCanvas* Canvas, bool bAdditionalViewFamily)
 {
-	UMediaTexture* ObjectTexture = GetThumbnailTextureFromObject(Object);
+	UTexture* ObjectTexture = GetThumbnailTextureFromObject(Object);
 	if (ObjectTexture != nullptr)
 	{
 		Super::Draw(ObjectTexture, X, Y, Width, Height, RenderTarget, Canvas, bAdditionalViewFamily);
 	}
 }
 
-UMediaTexture* UMediaSourceThumbnailRenderer::GetThumbnailTextureFromObject(UObject* Object) const
+UTexture* UMediaSourceThumbnailRenderer::GetThumbnailTextureFromObject(UObject* Object) const
 {
 	UMediaSource* MediaSource = Cast<UMediaSource>(Object);
 	if (MediaSource != nullptr)
 	{
-		UMediaTexture* Texture = MediaSource->GetThumbnail();
+		UTexture* Texture = MediaSource->GetThumbnail();
 		return Texture;
 	}
 	
diff --git a/Engine/Plugins/Media/MediaPlayerEditor/Source/MediaPlayerEditor/Private/Shared/MediaSourceThumbnailRenderer.h b/Engine/Plugins/Media/MediaPlayerEditor/Source/MediaPlayerEditor/Private/Shared/MediaSourceThumbnailRenderer.h
--- a/Engine/Plugins/Media/MediaPlayerEditor/Source/MediaPlayerEditor/Private/Shared/MediaSourceThumbnailRenderer.h
+++ b/Engine/Plugins/Media/MediaPlayerEditor/Source/MediaPlayerEditor/Private/Shared/MediaSourceThumbnailRenderer.h
@@ -6,7 +6,7 @@
 
 #include "MediaSourceThumbnailRenderer.generated.h"
 
-class UMediaTexture;
+class UTexture;
 
 /**
  * Renders thumbnails for media sources.
@@ -28,6 +28,6 @@ protected:
 	/**
 	 * Gets the thumnbail from the object if the object is a media source.
 	 */
-	virtual UMediaTexture* GetThumbnailTextureFromObject(UObject* Object) const;
+	virtual UTexture* GetThumbnailTextureFromObject(UObject* Object) const;
 
 };
diff --git a/Engine/Plugins/Media/MediaPlayerEditor/Source/MediaPlayerEditor/Private/Toolkits/MediaSourceEditorToolkit.cpp b/Engine/Plugins/Media/MediaPlayerEditor/Source/MediaPlayerEditor/Private/Toolkits/MediaSourceEditorToolkit.cpp
--- a/Engine/Plugins/Media/MediaPlayerEditor/Source/MediaPlayerEditor/Private/Toolkits/MediaSourceEditorToolkit.cpp
+++ b/Engine/Plugins/Media/MediaPlayerEditor/Source/MediaPlayerEditor/Private/Toolkits/MediaSourceEditorToolkit.cpp
@@ -4,14 +4,19 @@
 
 #include "Editor.h"
 #include "EditorReimportHandler.h"
+#include "Engine/TextureRenderTarget2D.h"
 #include "Framework/MultiBox/MultiBoxBuilder.h"
 #include "MediaPlayer.h"
 #include "MediaPlaylist.h"
 #include "MediaSource.h"
 #include "MediaTexture.h"
 #include "Models/MediaPlayerEditorCommands.h"
+#include "RenderGraphBuilder.h"
+#include "RenderGraphUtils.h"
+#include "ScreenPass.h"
 #include "SlateOptMacros.h"
 #include "Styling/AppStyle.h"
+#include "TextureResource.h"
 #include "Widgets/Docking/SDockTab.h"
 #include "Widgets/SMediaPlayerEditorDetails.h"
 #include "Widgets/SMediaPlayerEditorMediaDetails.h"
@@ -269,6 +274,11 @@ void FMediaSourceEditorToolkit::BindCommands()
 		FCanExecuteAction::CreateLambda([this]{ return MediaPlayer->IsReady() && MediaPlayer->SupportsRate(GetForwardRate(), false); })
 	);
 
+	ToolkitCommands->MapAction(
+		Commands.GenerateThumbnail,
+		FExecuteAction::CreateLambda([this] { GenerateThumbnail(); })
+	);
+
 	ToolkitCommands->MapAction(
 		Commands.NextMedia,
 		FExecuteAction::CreateLambda([this]{ MediaPlayer->Next(); }),
@@ -336,6 +346,7 @@ void FMediaSourceEditorToolkit::ExtendToolBar()
 			{
 				ToolbarBuilder.AddToolBarButton(FMediaPlayerEditorCommands::Get().OpenMedia);
 				ToolbarBuilder.AddToolBarButton(FMediaPlayerEditorCommands::Get().CloseMedia);
+				ToolbarBuilder.AddToolBarButton(FMediaPlayerEditorCommands::Get().GenerateThumbnail);
 			}
 			ToolbarBuilder.EndSection();
 		}
@@ -381,6 +392,84 @@ float FMediaSourceEditorToolkit::GetReverseRate() const
 	return 2.0f * Rate;
 }
 
+void FMediaSourceEditorToolkit::GenerateThumbnail()
+{
+	// Create render target.
+	UTextureRenderTarget2D* ThumbnailTexture =
+		NewObject<UTextureRenderTarget2D>(GetTransientPackage());
+	if (ThumbnailTexture != nullptr)
+	{
+		ThumbnailTexture->RenderTargetFormat = ETextureRenderTargetFormat::RTF_RGBA8;
+		ThumbnailTexture->ClearColor = FLinearColor::Black;
+		ThumbnailTexture->bAutoGenerateMips = false;
+		ThumbnailTexture->InitAutoFormat(MediaTexture->GetWidth(), MediaTexture->GetHeight());
+		ThumbnailTexture->UpdateResourceImmediate(true);
+
+		// Enqueue render command to copy the media texture to the render target.
+		ENQUEUE_RENDER_COMMAND(MediaSourceRenderThumbnail)(
+			[this, ThumbnailTexture](FRHICommandListImmediate& RHICmdList)
+		{
+			FTextureResource* DestResource = ThumbnailTexture->GetResource();
+			FTextureResource* SourceResource = MediaTexture->GetResource();
+			FRHITexture* SourceTexture = SourceResource->GetTextureRHI();
+			FRHITexture* DestTexture = DestResource->GetTextureRHI();
+
+			FRDGBuilder GraphBuilder(RHICmdList);
+			FRDGTextureRef RDGSourceTexture = RegisterExternalTexture(GraphBuilder, SourceTexture, TEXT("MediaSourceThumnbailSourceTexture"));
+			FRDGTextureRef RDGDestTexture = RegisterExternalTexture(GraphBuilder, DestTexture, TEXT("MediaSourceThumnbailDestTexture"));
+
+			FGlobalShaderMap* ShaderMap = GetGlobalShaderMap(GMaxRHIFeatureLevel);
+
+			TShaderMapRef<FScreenPassVS> VertexShader(ShaderMap);
+			TShaderMapRef<FCopyRectPS> PixelShader(ShaderMap);
+
+			FCopyRectPS::FParameters* PixelShaderParameters = GraphBuilder.AllocParameters<FCopyRectPS::FParameters>();
+			PixelShaderParameters->InputTexture = RDGSourceTexture;
+			PixelShaderParameters->InputSampler = TStaticSamplerState<SF_Bilinear, AM_Clamp, AM_Clamp, AM_Clamp>::GetRHI();
+			PixelShaderParameters->RenderTargets[0] = FRenderTargetBinding(RDGDestTexture, ERenderTargetLoadAction::ELoad);
+
+			ClearUnusedGraphResources(PixelShader, PixelShaderParameters);
+
+			FRHIBlendState* BlendState = FScreenPassPipelineState::FDefaultBlendState::GetRHI();
+			FRHIDepthStencilState* DepthStencilState = FScreenPassPipelineState::FDefaultDepthStencilState::GetRHI();
+
+			// Create the pipline state that will execute
+			const FScreenPassPipelineState PipelineState(VertexShader, PixelShader, BlendState, DepthStencilState);
+			GraphBuilder.AddPass(
+				RDG_EVENT_NAME("MediaSourceThumbnailCopy"),
+				PixelShaderParameters,
+				ERDGPassFlags::Raster,
+				[PipelineState, Extent = RDGSourceTexture->Desc.Extent, PixelShader, PixelShaderParameters](FRHICommandList& RHICmdList) {
+				PipelineState.Validate();
+
+				RHICmdList.SetViewport(0.0f, 0.0f, 0.0f, Extent.X, Extent.Y, 1.0f);
+				SetScreenPassPipelineState(RHICmdList, PipelineState);
+				SetShaderParameters(RHICmdList, PixelShader, PixelShader.GetPixelShader(), *PixelShaderParameters);
+
+				DrawRectangle(
+					RHICmdList,
+					0, 0, Extent.X, Extent.Y,
+					0, 0, Extent.X, Extent.Y,
+					Extent,
+					Extent,
+					PipelineState.VertexShader,
+					EDRF_UseTriangleOptimization);
+			});
+
+			GraphBuilder.Execute();
+
+		});
+
+		MediaSource->SetThumbnail(ThumbnailTexture);
+
+		// Trigger a thumbnail render.
+		FPropertyChangedEvent EmptyPropertyChangedEvent(nullptr);
+		FCoreUObjectDelegates::OnObjectPropertyChanged.Broadcast(MediaSource,
+			EmptyPropertyChangedEvent);
+		MediaSource->MarkPackageDirty();
+	}
+}
+
 
 /* FMediaSourceEditorToolkit callbacks
  *****************************************************************************/
diff --git a/Engine/Plugins/Media/MediaPlayerEditor/Source/MediaPlayerEditor/Private/Toolkits/MediaSourceEditorToolkit.h b/Engine/Plugins/Media/MediaPlayerEditor/Source/MediaPlayerEditor/Private/Toolkits/MediaSourceEditorToolkit.h
--- a/Engine/Plugins/Media/MediaPlayerEditor/Source/MediaPlayerEditor/Private/Toolkits/MediaSourceEditorToolkit.h
+++ b/Engine/Plugins/Media/MediaPlayerEditor/Source/MediaPlayerEditor/Private/Toolkits/MediaSourceEditorToolkit.h
@@ -100,6 +100,11 @@ protected:
 	 */
 	float GetReverseRate() const;
 
+	/**
+	 * Enqueues rendering commands to generate a thumbnail.
+	 */
+	void GenerateThumbnail();
+
 private:
 
 	/** Callback for spawning tabs. */
diff --git a/Engine/Source/Runtime/MediaAssets/Private/Assets/MediaSource.cpp b/Engine/Source/Runtime/MediaAssets/Private/Assets/MediaSource.cpp
--- a/Engine/Source/Runtime/MediaAssets/Private/Assets/MediaSource.cpp
+++ b/Engine/Source/Runtime/MediaAssets/Private/Assets/MediaSource.cpp
@@ -4,6 +4,7 @@
 
 #include "IMediaAssetsModule.h"
 #include "MediaAssetsPrivate.h"
+#include "MediaTexture.h"
 #include "Modules/ModuleManager.h"
 #include "Misc/Paths.h"
 #include "StreamMediaSource.h"
diff --git a/Engine/Source/Runtime/MediaAssets/Public/MediaSource.h b/Engine/Source/Runtime/MediaAssets/Public/MediaSource.h
--- a/Engine/Source/Runtime/MediaAssets/Public/MediaSource.h
+++ b/Engine/Source/Runtime/MediaAssets/Public/MediaSource.h
@@ -23,7 +23,7 @@
 class FVariant;
 class IMediaSourceRendererInterface;
 class UMediaSource;
-class UMediaTexture;
+class UTexture;
 struct FFrame;
 
 /** Delegate for creating a media source from a string. */
@@ -102,7 +102,12 @@ public:
 	/**
 	 * Gets our thumbnail texture, if any.
 	 */
-	UMediaTexture* GetThumbnail() const { return ThumbnailImage; }
+	UTexture* GetThumbnail() const { return ThumbnailImage; }
+
+	/**
+	 * Sets what the thumbnail texture should be.
+	 */
+	void SetThumbnail(UTexture* InTexture) { ThumbnailImage = InTexture; }
 
 #endif // WITH_EDITOR
 
@@ -186,7 +191,7 @@ private:
 
 	/** The thumbnail image.*/
 	UPROPERTY(Transient)
-	TObjectPtr<UMediaTexture> ThumbnailImage = nullptr;
+	TObjectPtr<UTexture> ThumbnailImage = nullptr;
 
 	/** Renders thumnbnails for us. */
 	UPROPERTY(Transient)
