MediaPlayerEditor: FMediaSourceActions::GetAssetDescription no longer loads the asset to get the description.

diff --git a/Engine/Plugins/Media/MediaPlayerEditor/Source/MediaPlayerEditor/Private/AssetTools/MediaSourceActions.cpp b/Engine/Plugins/Media/MediaPlayerEditor/Source/MediaPlayerEditor/Private/AssetTools/MediaSourceActions.cpp
--- a/Engine/Plugins/Media/MediaPlayerEditor/Source/MediaPlayerEditor/Private/AssetTools/MediaSourceActions.cpp
+++ b/Engine/Plugins/Media/MediaPlayerEditor/Source/MediaPlayerEditor/Private/AssetTools/MediaSourceActions.cpp
@@ -21,21 +21,17 @@ bool FMediaSourceActions::CanFilter()
 
 FText FMediaSourceActions::GetAssetDescription(const struct FAssetData& AssetData) const
 {
-	auto MediaSource = Cast<UMediaSource>(AssetData.GetAsset());
-
-	if (MediaSource != nullptr)
+	FString Url = AssetData.GetTagValueRef<FString>("Url");
+	if (Url.IsEmpty())
 	{
-		FString Url = MediaSource->GetUrl();
-
-		if (Url.IsEmpty())
-		{
-			return LOCTEXT("AssetTypeActions_MediaSourceMissing", "Warning: Missing settings detected!");
-		}
+		return LOCTEXT("AssetTypeActions_MediaSourceMissing", "Warning: Missing settings detected!");
+	}
 
-		if (!MediaSource->Validate())
-		{
-			return LOCTEXT("AssetTypeActions_MediaSourceInvalid", "Warning: Invalid settings detected!");
-		}
+	FString Validate = AssetData.GetTagValueRef<FString>("Validate");
+	bool bIsValidated = Validate.IsEmpty() || (Validate == TEXT("True"));
+	if (!bIsValidated)
+	{
+		return LOCTEXT("AssetTypeActions_MediaSourceInvalid", "Warning: Invalid settings detected!");
 	}
 
 	return FText::GetEmpty();
diff --git a/Engine/Source/Runtime/MediaAssets/Private/Assets/BaseMediaSource.cpp b/Engine/Source/Runtime/MediaAssets/Private/Assets/BaseMediaSource.cpp
--- a/Engine/Source/Runtime/MediaAssets/Private/Assets/BaseMediaSource.cpp
+++ b/Engine/Source/Runtime/MediaAssets/Private/Assets/BaseMediaSource.cpp
@@ -27,6 +27,9 @@ void UBaseMediaSource::GetAssetRegistryTags(TArray<FAssetRegistryTag>& OutTags)
 	{
 		OutTags.Add(FAssetRegistryTag("Url", Url, FAssetRegistryTag::TT_Alphabetical));
 	}
+	
+	OutTags.Add(FAssetRegistryTag("Validate", Validate() ? TEXT("True") : TEXT("False"),
+		FAssetRegistryTag::TT_Alphabetical));
 	Super::GetAssetRegistryTags(OutTags);
 }
 
