MediaPlayerEditor: Added context menu item for media sources to generate thumbnails.

diff --git a/Engine/Plugins/Media/MediaPlayerEditor/Source/MediaPlayerEditor/Private/MediaPlayerEditorModule.cpp b/Engine/Plugins/Media/MediaPlayerEditor/Source/MediaPlayerEditor/Private/MediaPlayerEditorModule.cpp
--- a/Engine/Plugins/Media/MediaPlayerEditor/Source/MediaPlayerEditor/Private/MediaPlayerEditorModule.cpp
+++ b/Engine/Plugins/Media/MediaPlayerEditor/Source/MediaPlayerEditor/Private/MediaPlayerEditorModule.cpp
@@ -3,6 +3,7 @@
 #include "MediaPlayerEditorModule.h"
 
 #include "ComponentVisualizer.h"
+#include "ContentBrowserMenuContexts.h"
 #include "Modules/ModuleManager.h"
 #include "AssetToolsModule.h"
 #include "UObject/UObjectHash.h"
@@ -109,6 +110,7 @@ public:
 		RegisterMenuExtensions();
 		RegisterThumbnailRenderers();
 		RegisterVisualizers();
+		ExtendContentMenu();
 	}
 
 	virtual void ShutdownModule() override
@@ -300,6 +302,55 @@ protected:
 
 private:
 
+	void ExtendContentMenu()
+	{
+		UToolMenu* Menu = UToolMenus::Get()->ExtendMenu("ContentBrowser.AssetContextMenu");
+
+		FToolMenuSection& Section = Menu->AddDynamicSection("MediaSource", FNewToolMenuDelegate::CreateLambda([this](UToolMenu* ToolMenu)
+		{
+			if (const UContentBrowserAssetContextMenuContext* Context = UContentBrowserAssetContextMenuContext::FindContextWithAssets(*ToolMenu))
+			{
+				TArray<UMediaSource*> MediaSources;
+				for (const FAssetData& AssetData : Context->SelectedAssets)
+				{
+					if (UMediaSource* MediaSource = Cast<UMediaSource>(AssetData.GetAsset()))
+					{
+						MediaSources.Add(MediaSource);
+					}
+				}
+
+				if (MediaSources.Num() > 0)
+				{
+					const TAttribute<FText> Label = LOCTEXT("GenerateThumbnail", "Generate Thumbnail");
+					const TAttribute<FText> ToolTip = LOCTEXT("GenerateThumbnail_Tooltip",
+						"Generate a thumbnail for this asset.");
+					const FSlateIcon Icon = FSlateIcon(FAppStyle::GetAppStyleSetName(),
+						"AssetEditor.SaveThumbnail");
+					const FToolMenuExecuteAction UIAction =
+						FToolMenuExecuteAction::CreateLambda([MediaSources]
+						(const FToolMenuContext& InContext)
+					{
+						for (UMediaSource* MediaSource : MediaSources)
+						{
+							MediaSource->GenerateThumbnail();
+							MediaSource->MarkPackageDirty();
+						}
+					});
+
+					const FName SectionName = TEXT("Media");
+					FToolMenuSection* Section = ToolMenu->FindSection(SectionName);
+					if (!Section)
+					{
+						Section = &(ToolMenu->AddSection(SectionName,
+							LOCTEXT("MediaSectionLabel", "Media")));
+					}
+
+					Section->AddMenuEntry("Media_GenerateThumbnail", Label, ToolTip, Icon, UIAction);
+				}
+			}
+		}), FToolMenuInsert(NAME_None, EToolMenuInsertType::Default));
+	}
+
 	void HandleEditorBeginPIE(bool bIsSimulating)
 	{
 		for (TObjectIterator<UMediaPlayer> It; It; ++It)
