Editor - Fixing a regression to the new no-longer loading things on right click.  This code in media player editor module caused every asset to load regardless of what it actually was.

diff --git a/Engine/Plugins/Media/MediaPlayerEditor/Source/MediaPlayerEditor/Private/MediaPlayerEditorModule.cpp b/Engine/Plugins/Media/MediaPlayerEditor/Source/MediaPlayerEditor/Private/MediaPlayerEditorModule.cpp
--- a/Engine/Plugins/Media/MediaPlayerEditor/Source/MediaPlayerEditor/Private/MediaPlayerEditorModule.cpp
+++ b/Engine/Plugins/Media/MediaPlayerEditor/Source/MediaPlayerEditor/Private/MediaPlayerEditorModule.cpp
@@ -304,51 +304,31 @@ private:
 
 	void ExtendContentMenu()
 	{
-		UToolMenu* Menu = UToolMenus::Get()->ExtendMenu("ContentBrowser.AssetContextMenu");
+		UToolMenu* Menu = UE::ContentBrowser::ExtendToolMenu_AssetContextMenu(UMediaSource::StaticClass());
 
-		FToolMenuSection& Section = Menu->AddDynamicSection("MediaSource", FNewToolMenuDelegate::CreateLambda([this](UToolMenu* ToolMenu)
+		FToolMenuSection& Section = Menu->FindOrAddSection("GetAssetActions");
+		Section.AddDynamicEntry(NAME_None, FNewToolMenuSectionDelegate::CreateLambda([](FToolMenuSection& InSection)
 		{
-			if (const UContentBrowserAssetContextMenuContext* Context = UContentBrowserAssetContextMenuContext::FindContextWithAssets(*ToolMenu))
+			if (const UContentBrowserAssetContextMenuContext* Context = UContentBrowserAssetContextMenuContext::FindContextWithAssets(InSection))
 			{
-				TArray<UMediaSource*> MediaSources;
-				for (const FAssetData& AssetData : Context->SelectedAssets)
-				{
-					if (UMediaSource* MediaSource = Cast<UMediaSource>(AssetData.GetAsset()))
-					{
-						MediaSources.Add(MediaSource);
-					}
-				}
+				const TAttribute<FText> Label = LOCTEXT("GenerateThumbnail", "Generate Thumbnail");
+				const TAttribute<FText> ToolTip = LOCTEXT("GenerateThumbnail_Tooltip", "Generate a thumbnail for this asset.");
+				const FSlateIcon Icon = FSlateIcon(FAppStyle::GetAppStyleSetName(), "AssetEditor.SaveThumbnail");;
 
-				if (MediaSources.Num() > 0)
+				FToolUIAction UIAction;
+				UIAction.ExecuteAction = FToolMenuExecuteAction::CreateLambda([](const FToolMenuContext& InContext)
 				{
-					const TAttribute<FText> Label = LOCTEXT("GenerateThumbnail", "Generate Thumbnail");
-					const TAttribute<FText> ToolTip = LOCTEXT("GenerateThumbnail_Tooltip",
-						"Generate a thumbnail for this asset.");
-					const FSlateIcon Icon = FSlateIcon(FAppStyle::GetAppStyleSetName(),
-						"AssetEditor.SaveThumbnail");
-					const FToolMenuExecuteAction UIAction =
-						FToolMenuExecuteAction::CreateLambda([MediaSources]
-						(const FToolMenuContext& InContext)
-					{
-						for (UMediaSource* MediaSource : MediaSources)
-						{
-							MediaSource->GenerateThumbnail();
-							MediaSource->MarkPackageDirty();
-						}
-					});
-
-					const FName SectionName = TEXT("Media");
-					FToolMenuSection* Section = ToolMenu->FindSection(SectionName);
-					if (!Section)
+					const UContentBrowserAssetContextMenuContext* Context = UContentBrowserAssetContextMenuContext::FindContextWithAssets(InContext);
+
+					for (UMediaSource* MediaSource : Context->LoadSelectedObjects<UMediaSource>())
 					{
-						Section = &(ToolMenu->AddSection(SectionName,
-							LOCTEXT("MediaSectionLabel", "Media")));
+						MediaSource->GenerateThumbnail();
+						MediaSource->MarkPackageDirty();
 					}
-
-					Section->AddMenuEntry("Media_GenerateThumbnail", Label, ToolTip, Icon, UIAction);
-				}
+				});
+				InSection.AddMenuEntry("Media_GenerateThumbnail", Label, ToolTip, Icon, UIAction);
 			}
-		}), FToolMenuInsert(NAME_None, EToolMenuInsertType::Default));
+		}));
 	}
 
 	void HandleEditorBeginPIE(bool bIsSimulating)
