Fixing bad include path

diff --git a/Engine/Plugins/Media/MediaPlayerEditor/Source/MediaPlayerEditor/Private/MediaPlayerEditorModule.cpp b/Engine/Plugins/Media/MediaPlayerEditor/Source/MediaPlayerEditor/Private/MediaPlayerEditorModule.cpp
--- a/Engine/Plugins/Media/MediaPlayerEditor/Source/MediaPlayerEditor/Private/MediaPlayerEditorModule.cpp
+++ b/Engine/Plugins/Media/MediaPlayerEditor/Source/MediaPlayerEditor/Private/MediaPlayerEditorModule.cpp
@@ -22,7 +22,7 @@
 #include "IMediaAssetsModule.h"
 #include "MediaPlayer.h"
 #include "MediaSoundComponent.h"
-#include "MediaSourceRenderer.h"
+#include "Shared/MediaSourceRenderer.h"
 #include "MediaTexture.h"
 #include "PlatformMediaSource.h"
 
