MediaPlayerEditor: Enabled looping for media source preview playback.
#rnx

diff --git a/Engine/Plugins/Media/MediaPlayerEditor/Source/MediaPlayerEditor/Private/Toolkits/MediaSourceEditorToolkit.cpp b/Engine/Plugins/Media/MediaPlayerEditor/Source/MediaPlayerEditor/Private/Toolkits/MediaSourceEditorToolkit.cpp
--- a/Engine/Plugins/Media/MediaPlayerEditor/Source/MediaPlayerEditor/Private/Toolkits/MediaSourceEditorToolkit.cpp
+++ b/Engine/Plugins/Media/MediaPlayerEditor/Source/MediaPlayerEditor/Private/Toolkits/MediaSourceEditorToolkit.cpp
@@ -78,7 +78,7 @@ void FMediaSourceEditorToolkit::Initialize(UMediaSource* InMediaSource, const ET
 	GEditor->RegisterForUndo(this);
 
 	MediaPlayer = NewObject<UMediaPlayer>(GetTransientPackage());
-	MediaPlayer->SetLooping(false);
+	MediaPlayer->SetLooping(true);
 	MediaPlayer->PlayOnOpen = true;
 
 	MediaTexture = NewObject<UMediaTexture>(GetTransientPackage(), NAME_None, RF_Transient | RF_Public);
