MediaPlayerEditor: LWC truncation fixes.

#rnx

diff --git a/Engine/Plugins/Media/MediaPlayerEditor/Source/MediaPlayerEditor/Private/Shared/MediaSourceThumbnailRenderer.cpp b/Engine/Plugins/Media/MediaPlayerEditor/Source/MediaPlayerEditor/Private/Shared/MediaSourceThumbnailRenderer.cpp
--- a/Engine/Plugins/Media/MediaPlayerEditor/Source/MediaPlayerEditor/Private/Shared/MediaSourceThumbnailRenderer.cpp
+++ b/Engine/Plugins/Media/MediaPlayerEditor/Source/MediaPlayerEditor/Private/Shared/MediaSourceThumbnailRenderer.cpp
@@ -16,8 +16,8 @@ void UMediaSourceThumbnailRenderer::GetThumbnailSize(UObject* Object, float Zoom
 	UTexture* ObjectTexture = GetThumbnailTextureFromObject(Object);
 	if (ObjectTexture != nullptr)
 	{
-		OutWidth = Zoom * ObjectTexture->GetSurfaceWidth();
-		OutHeight = Zoom * ObjectTexture->GetSurfaceHeight();
+		OutWidth = static_cast<uint32>(Zoom * ObjectTexture->GetSurfaceWidth());
+		OutHeight = static_cast<uint32>(Zoom * ObjectTexture->GetSurfaceHeight());
 	}
 	else
 	{
diff --git a/Engine/Plugins/Media/MediaPlayerEditor/Source/MediaPlayerEditor/Private/Toolkits/MediaSourceEditorToolkit.cpp b/Engine/Plugins/Media/MediaPlayerEditor/Source/MediaPlayerEditor/Private/Toolkits/MediaSourceEditorToolkit.cpp
--- a/Engine/Plugins/Media/MediaPlayerEditor/Source/MediaPlayerEditor/Private/Toolkits/MediaSourceEditorToolkit.cpp
+++ b/Engine/Plugins/Media/MediaPlayerEditor/Source/MediaPlayerEditor/Private/Toolkits/MediaSourceEditorToolkit.cpp
@@ -443,14 +443,14 @@ void FMediaSourceEditorToolkit::GenerateThumbnail()
 				[PipelineState, Extent = RDGSourceTexture->Desc.Extent, PixelShader, PixelShaderParameters](FRHICommandList& RHICmdList) {
 				PipelineState.Validate();
 
-				RHICmdList.SetViewport(0.0f, 0.0f, 0.0f, Extent.X, Extent.Y, 1.0f);
+				RHICmdList.SetViewport(0.0f, 0.0f, 0.0f, static_cast<float>(Extent.X), static_cast<float>(Extent.Y), 1.0f);
 				SetScreenPassPipelineState(RHICmdList, PipelineState);
 				SetShaderParameters(RHICmdList, PixelShader, PixelShader.GetPixelShader(), *PixelShaderParameters);
 
 				DrawRectangle(
 					RHICmdList,
-					0, 0, Extent.X, Extent.Y,
-					0, 0, Extent.X, Extent.Y,
+					0.0f, 0.0f, static_cast<float>(Extent.X), static_cast<float>(Extent.Y),
+					0.0f, 0.0f, static_cast<float>(Extent.X), static_cast<float>(Extent.Y),
 					Extent,
 					Extent,
 					PipelineState.VertexShader,
diff --git a/Engine/Plugins/Media/MediaPlayerEditor/Source/MediaPlayerEditor/Private/Widgets/SMediaImage.cpp b/Engine/Plugins/Media/MediaPlayerEditor/Source/MediaPlayerEditor/Private/Widgets/SMediaImage.cpp
--- a/Engine/Plugins/Media/MediaPlayerEditor/Source/MediaPlayerEditor/Private/Widgets/SMediaImage.cpp
+++ b/Engine/Plugins/Media/MediaPlayerEditor/Source/MediaPlayerEditor/Private/Widgets/SMediaImage.cpp
@@ -92,7 +92,7 @@ void SMediaImage::Tick(const FGeometry& InAllottedGeometry, const double InCurre
 	if (BrushImageSize.IsSet())
 	{
 		FVector2D Size = BrushImageSize.Get();
-		MaterialBrush->ImageSize.X = Size.X;
-		MaterialBrush->ImageSize.Y = Size.Y;
+		MaterialBrush->ImageSize.X = static_cast<float>(Size.X);
+		MaterialBrush->ImageSize.Y = static_cast<float>(Size.Y);
 	}
 }
diff --git a/Engine/Plugins/Media/MediaPlayerEditor/Source/MediaPlayerEditor/Private/Widgets/SMediaPlayerEditorCache.cpp b/Engine/Plugins/Media/MediaPlayerEditor/Source/MediaPlayerEditor/Private/Widgets/SMediaPlayerEditorCache.cpp
--- a/Engine/Plugins/Media/MediaPlayerEditor/Source/MediaPlayerEditor/Private/Widgets/SMediaPlayerEditorCache.cpp
+++ b/Engine/Plugins/Media/MediaPlayerEditor/Source/MediaPlayerEditor/Private/Widgets/SMediaPlayerEditorCache.cpp
@@ -80,7 +80,7 @@ void SMediaPlayerEditorCache::DrawPlayerPosition(FTimespan Time, const FGeometry
 	const FTimespan Duration = MediaPlayer->GetDuration();
 	const float MarkerSize = PositionMarkerSize.Get();
 
-	const float DrawOffset = FTimespan::Ratio(Time, Duration) * AllottedGeometry.Size.X - 0.5f * MarkerSize;
+	const float DrawOffset = static_cast<float>(FTimespan::Ratio(Time, Duration) * AllottedGeometry.Size.X - 0.5f * MarkerSize);
 
 	FSlateDrawElement::MakeBox(
 		OutDrawElements,
@@ -132,8 +132,8 @@ void SMediaPlayerEditorCache::DrawSampleStates(const TRangeSet<FTimespan>& Range
 
 	for (auto& Range : Ranges)
 	{
-		const float DrawOffset = FMath::RoundToNegativeInfinity(FTimespan::Ratio(Range.GetLowerBoundValue(), Duration) * AllottedGeometry.Size.X);
-		const float DrawSize = FMath::RoundToPositiveInfinity(FTimespan::Ratio(Range.Size<FTimespan>(), Duration) * AllottedGeometry.Size.X);
+		const float DrawOffset = static_cast<float>(FMath::RoundToNegativeInfinity(FTimespan::Ratio(Range.GetLowerBoundValue(), Duration) * AllottedGeometry.Size.X));
+		const float DrawSize = static_cast<float>(FMath::RoundToPositiveInfinity(FTimespan::Ratio(Range.Size<FTimespan>(), Duration) * AllottedGeometry.Size.X));
 		const float BarHeight = ProgressBarHeight.Get();
 
 		FSlateDrawElement::MakeBox(
diff --git a/Engine/Plugins/Media/MediaPlayerEditor/Source/MediaPlayerEditor/Private/Widgets/SMediaPlayerEditorMediaDetails.cpp b/Engine/Plugins/Media/MediaPlayerEditor/Source/MediaPlayerEditor/Private/Widgets/SMediaPlayerEditorMediaDetails.cpp
--- a/Engine/Plugins/Media/MediaPlayerEditor/Source/MediaPlayerEditor/Private/Widgets/SMediaPlayerEditorMediaDetails.cpp
+++ b/Engine/Plugins/Media/MediaPlayerEditor/Source/MediaPlayerEditor/Private/Widgets/SMediaPlayerEditorMediaDetails.cpp
@@ -160,8 +160,8 @@ void SMediaPlayerEditorMediaDetails::UpdateDetails()
 				: LOCTEXT("MethodStreamed", "Streamed"));
 		NumMips = MediaTexture->GetTextureNumMips();
 		ResourceSize = (MediaTexture->GetResourceSizeBytes(EResourceSizeMode::Exclusive) + 512) / 1024;
-		SurfaceWidth = MediaTexture->GetSurfaceWidth();
-		SurfaceHeight = MediaTexture->GetSurfaceHeight();
+		SurfaceWidth = static_cast<int32>(MediaTexture->GetSurfaceWidth());
+		SurfaceHeight = static_cast<int32>(MediaTexture->GetSurfaceHeight());
 	}
 
 	// Update text.
diff --git a/Engine/Plugins/Media/MediaPlayerEditor/Source/MediaPlayerEditor/Private/Widgets/SMediaPlayerEditorViewer.cpp b/Engine/Plugins/Media/MediaPlayerEditor/Source/MediaPlayerEditor/Private/Widgets/SMediaPlayerEditorViewer.cpp
--- a/Engine/Plugins/Media/MediaPlayerEditor/Source/MediaPlayerEditor/Private/Widgets/SMediaPlayerEditorViewer.cpp
+++ b/Engine/Plugins/Media/MediaPlayerEditor/Source/MediaPlayerEditor/Private/Widgets/SMediaPlayerEditorViewer.cpp
@@ -381,7 +381,7 @@ void SMediaPlayerEditorViewer::Construct(const FArguments& InArgs, UMediaPlayer&
 													})
 													.OnMouseCaptureBegin_Lambda([this]()
 													{
-														ScrubValue = FTimespan::Ratio(MediaPlayer->GetDisplayTime(), MediaPlayer->GetDuration());
+														ScrubValue = static_cast<float>(FTimespan::Ratio(MediaPlayer->GetDisplayTime(), MediaPlayer->GetDuration()));
 
 														if (MediaPlayer->SupportsScrubbing())
 														{
@@ -416,7 +416,7 @@ void SMediaPlayerEditorViewer::Construct(const FArguments& InArgs, UMediaPlayer&
 														}
 
 														return (MediaPlayer->GetDuration() > FTimespan::Zero())
-															? FTimespan::Ratio(MediaPlayer->GetDisplayTime(), MediaPlayer->GetDuration())
+															? static_cast<float>(FTimespan::Ratio(MediaPlayer->GetDisplayTime(), MediaPlayer->GetDuration()))
 															: 0.0f;
 													})
 													.Visibility_Lambda([this]() -> EVisibility {
