Fix Level Snapshots not restoring Media Plates playlists.
- Added UMediaPlateComponent::MediaPlaylist to the Level Snapshot list of allowed properties.
- Media Plate playlists editor now transacts modifications. This allows Level Snapshots to detect changes done to it.

diff --git a/Engine/Plugins/Media/MediaPlayerEditor/Source/MediaPlayerEditor/Private/Widgets/SMediaPlaylistEditorTracks.cpp b/Engine/Plugins/Media/MediaPlayerEditor/Source/MediaPlayerEditor/Private/Widgets/SMediaPlaylistEditorTracks.cpp
--- a/Engine/Plugins/Media/MediaPlayerEditor/Source/MediaPlayerEditor/Private/Widgets/SMediaPlaylistEditorTracks.cpp
+++ b/Engine/Plugins/Media/MediaPlayerEditor/Source/MediaPlayerEditor/Private/Widgets/SMediaPlaylistEditorTracks.cpp
@@ -12,6 +12,8 @@
 #include "Widgets/Input/SButton.h"
 #include "Widgets/Layout/SScrollBox.h"
 #include "Styling/AppStyle.h"
+#include "ScopedTransaction.h"
+#include "UObject/UObjectGlobals.h"
 
 #define LOCTEXT_NAMESPACE "SMediaPlaylistEditorTracks"
 
@@ -59,9 +61,16 @@ void SMediaPlaylistEditorTracks::Construct(const FArguments& InArgs, UMediaPlayl
 				]
 		];
 
+	FCoreUObjectDelegates::OnObjectTransacted.AddSP(this, &SMediaPlaylistEditorTracks::OnObjectTransacted);
 	RefreshPlaylist();
 }
 
+SMediaPlaylistEditorTracks::~SMediaPlaylistEditorTracks()
+{
+	// AddSP is technically safe from dangling pointers but let's be nice and remove ourselves.
+	FCoreUObjectDelegates::OnObjectTransacted.RemoveAll(this);
+}
+
 void SMediaPlaylistEditorTracks::RefreshPlaylist()
 {
 	SourcesContainer->ClearChildren();
@@ -104,8 +113,11 @@ void SMediaPlaylistEditorTracks::RefreshPlaylist()
 									UMediaPlaylist* MediaPlaylist = MediaPlaylistPtr.Get();
 									if (MediaPlaylist != nullptr)
 									{
+#if WITH_EDITOR
+                                    	FScopedTransaction Transaction(LOCTEXT("Array.InsertMediaSource", "Insert Media Source"));
+                                    	MediaPlaylist->Modify();
+#endif
 										MediaPlaylist->Insert(nullptr, Index);
-										MediaPlaylist->MarkPackageDirty();
 										RefreshPlaylist();
 									}
 								}),
@@ -114,8 +126,11 @@ void SMediaPlaylistEditorTracks::RefreshPlaylist()
 									UMediaPlaylist* MediaPlaylist = MediaPlaylistPtr.Get();
 									if (MediaPlaylist != nullptr)
 									{
+#if WITH_EDITOR
+										FScopedTransaction Transaction(LOCTEXT("Array.RemoveMediaSource", "Remove Media Source"));
+										MediaPlaylist->Modify();
+#endif
 										MediaPlaylist->RemoveAt(Index);
-										MediaPlaylist->MarkPackageDirty();
 										RefreshPlaylist();
 									}
 								}),
@@ -124,8 +139,11 @@ void SMediaPlaylistEditorTracks::RefreshPlaylist()
 									UMediaPlaylist* MediaPlaylist = MediaPlaylistPtr.Get();
 									if (MediaPlaylist != nullptr)
 									{
+#if WITH_EDITOR
+										FScopedTransaction Transaction(LOCTEXT("Array.DuplicateMediaSource", "Duplicate Media Source"));
+										MediaPlaylist->Modify();
+#endif
 										MediaPlaylist->Insert(MediaPlaylist->Get(Index), Index);
-										MediaPlaylist->MarkPackageDirty();
 										RefreshPlaylist();
 									}
 								}))
@@ -140,8 +158,11 @@ void SMediaPlaylistEditorTracks::AddToPlaylist()
 	UMediaPlaylist* MediaPlaylist = MediaPlaylistPtr.Get();
 	if (MediaPlaylist != nullptr)
 	{
+#if WITH_EDITOR
+		FScopedTransaction Transaction(LOCTEXT("AddToPlaylist", "Add Media Source"));
+		MediaPlaylist->Modify();
+#endif
 		MediaPlaylist->Insert(nullptr, MediaPlaylist->Num());
-		MediaPlaylist->MarkPackageDirty();
 		RefreshPlaylist();
 	}
 }
@@ -169,12 +190,20 @@ void SMediaPlaylistEditorTracks::OnMediaSourceChanged(const FAssetData& AssetDat
 	if (MediaPlaylist != nullptr)
 	{
 		UMediaSource* MediaSource = Cast<UMediaSource>(AssetData.GetAsset());
-		
+#if WITH_EDITOR
+		FScopedTransaction Transaction(LOCTEXT("ChangeMediaSource", "Change Media Source"));
+		MediaPlaylist->Modify();
+#endif	
 		MediaPlaylist->Replace(Index, MediaSource);
-		MediaPlaylist->MarkPackageDirty();
 	}
 }
 
-
+void SMediaPlaylistEditorTracks::OnObjectTransacted(UObject* Object, const FTransactionObjectEvent& TransactionObjectEvent)
+{
+	if (Object == MediaPlaylistPtr.Get())
+	{
+		RefreshPlaylist();
+	}
+}
 
 #undef LOCTEXT_NAMESPACE
diff --git a/Engine/Plugins/Media/MediaPlayerEditor/Source/MediaPlayerEditor/Public/Widgets/SMediaPlaylistEditorTracks.h b/Engine/Plugins/Media/MediaPlayerEditor/Source/MediaPlayerEditor/Public/Widgets/SMediaPlaylistEditorTracks.h
--- a/Engine/Plugins/Media/MediaPlayerEditor/Source/MediaPlayerEditor/Public/Widgets/SMediaPlaylistEditorTracks.h
+++ b/Engine/Plugins/Media/MediaPlayerEditor/Source/MediaPlayerEditor/Public/Widgets/SMediaPlaylistEditorTracks.h
@@ -22,7 +22,7 @@ public:
 	SLATE_END_ARGS()
 
 public:
-
+	
 	/**
 	 * Construct this widget
 	 *
@@ -31,6 +31,7 @@ public:
 	 * @param InStyleSet The style set to use.
 	 */
 	void Construct(const FArguments& InArgs, UMediaPlaylist* InMediaPlaylist, const TSharedRef<ISlateStyle>& InStyle);
+	virtual ~SMediaPlaylistEditorTracks() override;
 
 protected:
 
@@ -64,4 +65,6 @@ protected:
 	 */
 	void OnMediaSourceChanged(const FAssetData& AssetData, int32 Index);
 
+	/** Engine callback to detect when the media is changed. */
+	void OnObjectTransacted(UObject* Object, const FTransactionObjectEvent& TransactionObjectEvent);
 };
diff --git a/Engine/Plugins/VirtualProduction/LevelSnapshots/Source/LevelSnapshots/Private/SpecialSupport/EngineTypesRestorationFence.cpp b/Engine/Plugins/VirtualProduction/LevelSnapshots/Source/LevelSnapshots/Private/SpecialSupport/EngineTypesRestorationFence.cpp
--- a/Engine/Plugins/VirtualProduction/LevelSnapshots/Source/LevelSnapshots/Private/SpecialSupport/EngineTypesRestorationFence.cpp
+++ b/Engine/Plugins/VirtualProduction/LevelSnapshots/Source/LevelSnapshots/Private/SpecialSupport/EngineTypesRestorationFence.cpp
@@ -51,6 +51,21 @@ namespace UE::LevelSnapshots::Private::EngineTypesRestorationFence
 		}
 	}
 
+	static void AddMediaPlateSupport(ILevelSnapshotsModule& Module)
+	{
+		UStruct* MediaPlateComponentClass = FindObject<UStruct>(nullptr, TEXT("/Script/MediaPlate.MediaPlateComponent"));
+		if (!MediaPlateComponentClass)
+		{
+			return;
+		}
+		
+		const FProperty* MediaPlaylistProperty = MediaPlateComponentClass->FindPropertyByName(TEXT("MediaPlaylist"));
+		if (ensure(MediaPlaylistProperty))
+		{
+			Module.AddExplicitilySupportedProperties({ MediaPlaylistProperty });
+		}
+	}
+
 	static void DisableIrrelevantBrushSubobjects(ILevelSnapshotsModule& Module)
 	{
 #if WITH_EDITORONLY_DATA
@@ -117,6 +132,7 @@ namespace UE::LevelSnapshots::Private::EngineTypesRestorationFence
 		// Enable / disable troublesome properties
 		AddSoftObjectPathSupport(Module);
 		AddAttachParentSupport(Module);
+		AddMediaPlateSupport(Module);
 		DisableIrrelevantBrushSubobjects(Module);
 		DisableIrrelevantWorldSettings(Module);
 		DisableIrrelevantMaterialInstanceProperties(Module);
