Media: Prevent media thumbnail renderer from filling up the global cache with image media.

diff --git a/Engine/Plugins/Media/MediaPlayerEditor/Source/MediaPlayerEditor/Private/Shared/MediaSourceRenderer.cpp b/Engine/Plugins/Media/MediaPlayerEditor/Source/MediaPlayerEditor/Private/Shared/MediaSourceRenderer.cpp
--- a/Engine/Plugins/Media/MediaPlayerEditor/Source/MediaPlayerEditor/Private/Shared/MediaSourceRenderer.cpp
+++ b/Engine/Plugins/Media/MediaPlayerEditor/Source/MediaPlayerEditor/Private/Shared/MediaSourceRenderer.cpp
@@ -33,16 +33,23 @@ UMediaTexture* UMediaSourceRenderer::Open(UMediaSource* InMediaSource)
 		MediaTexture->UpdateResource();
 		MediaSource = InMediaSource;
 
+		// For image media, we avoid filling the global cache which will needlessly hold onto frame data.
+		FMediaSourceCacheSettings OriginalCacheSettings;
+		MediaSource->GetCacheSettings(OriginalCacheSettings);
+		MediaSource->SetCacheSettings(FMediaSourceCacheSettings{ true, 0.2f });
+
 		// Start playing the media.
 		bIsSeekActive = false;
 		FMediaPlayerOptions Options;
 		Options.PlayOnOpen = EMediaPlayerOptionBooleanOverride::Enabled;
 		Options.Loop = EMediaPlayerOptionBooleanOverride::Disabled;
-		bool bIsPlaying = MediaPlayer->OpenSourceWithOptions(InMediaSource, Options);
+		bool bIsPlaying = MediaPlayer->OpenSourceWithOptions(MediaSource, Options);
 		if (bIsPlaying == false)
 		{
 			Close();
 		}
+
+		MediaSource->SetCacheSettings(OriginalCacheSettings);
 	}
 
 	return MediaTexture.Get();
