Fix MediaImage / MediaPlate pre-visualization with Substrate.

diff --git a/Engine/Plugins/Media/MediaPlayerEditor/Source/MediaPlayerEditor/Private/Widgets/SMediaImage.cpp b/Engine/Plugins/Media/MediaPlayerEditor/Source/MediaPlayerEditor/Private/Widgets/SMediaImage.cpp
--- a/Engine/Plugins/Media/MediaPlayerEditor/Source/MediaPlayerEditor/Private/Widgets/SMediaImage.cpp
+++ b/Engine/Plugins/Media/MediaPlayerEditor/Source/MediaPlayerEditor/Private/Widgets/SMediaImage.cpp
@@ -6,6 +6,7 @@
 #include "MaterialDomain.h"
 #include "Materials/Material.h"
 #include "Materials/MaterialExpressionTextureSample.h"
+#include "Materials/MaterialExpressionSubstrate.h"
 #include "Styling/SlateBrush.h"
 #include "UObject/Package.h"
 #include "Widgets/Layout/SScaleBox.h"
@@ -40,24 +41,48 @@ void SMediaImage::Construct(const FArguments& InArgs, UTexture* InTexture)
 			}
 
 			FExpressionOutput& Output = TextureSampler->GetOutputs()[0];
-			FExpressionInput& Input = MaterialEditorOnly->EmissiveColor;
+			if (Substrate::IsSubstrateEnabled())
 			{
-				Input.Expression = TextureSampler;
-				Input.Mask = Output.Mask;
-				Input.MaskR = Output.MaskR;
-				Input.MaskG = Output.MaskG;
-				Input.MaskB = Output.MaskB;
-				Input.MaskA = Output.MaskA;
-			}
+				UMaterialExpressionSubstrateUI* UINode = NewObject<UMaterialExpressionSubstrateUI>(Material);
+				UINode->Material = Material;
+
+				UINode->Color.Expression = TextureSampler;
+				UINode->Color.Mask = Output.Mask;
+				UINode->Color.MaskR = Output.MaskR;
+				UINode->Color.MaskG = Output.MaskG;
+				UINode->Color.MaskB = Output.MaskB;
+				UINode->Color.MaskA = Output.MaskA;
+
+				UINode->Opacity.Expression = TextureSampler;
+				UINode->Opacity.Mask = Output.Mask;
+				UINode->Opacity.MaskR = 0;
+				UINode->Opacity.MaskG = 0;
+				UINode->Opacity.MaskB = 0;
+				UINode->Opacity.MaskA = 1;
 
-			FExpressionInput& Opacity = MaterialEditorOnly->Opacity;
+				MaterialEditorOnly->FrontMaterial.Connect(0, UINode);
+			}
+			else
 			{
-				Opacity.Expression = TextureSampler;
-				Opacity.Mask = Output.Mask;
-				Opacity.MaskR = 0;
-				Opacity.MaskG = 0;
-				Opacity.MaskB = 0;
-				Opacity.MaskA = 1;
+				FExpressionInput& Input = MaterialEditorOnly->EmissiveColor;
+				{
+					Input.Expression = TextureSampler;
+					Input.Mask = Output.Mask;
+					Input.MaskR = Output.MaskR;
+					Input.MaskG = Output.MaskG;
+					Input.MaskB = Output.MaskB;
+					Input.MaskA = Output.MaskA;
+				}
+
+				FExpressionInput& Opacity = MaterialEditorOnly->Opacity;
+				{
+					Opacity.Expression = TextureSampler;
+					Opacity.Mask = Output.Mask;
+					Opacity.MaskR = 0;
+					Opacity.MaskG = 0;
+					Opacity.MaskB = 0;
+					Opacity.MaskA = 1;
+				}
 			}
 
 			Material->BlendMode = BLEND_AlphaComposite;
