Fixing up usage of internal include path in MFW editor module
#rnx

diff --git a/Engine/Plugins/Media/MediaPlayerEditor/Source/MediaPlayerEditor/Private/Toolkits/MediaSourceEditorToolkit.cpp b/Engine/Plugins/Media/MediaPlayerEditor/Source/MediaPlayerEditor/Private/Toolkits/MediaSourceEditorToolkit.cpp
--- a/Engine/Plugins/Media/MediaPlayerEditor/Source/MediaPlayerEditor/Private/Toolkits/MediaSourceEditorToolkit.cpp
+++ b/Engine/Plugins/Media/MediaPlayerEditor/Source/MediaPlayerEditor/Private/Toolkits/MediaSourceEditorToolkit.cpp
@@ -25,8 +25,8 @@
 #include "Widgets/SMediaPlayerEditorViewer.h"
 #include "Widgets/SMediaSourceEditorDetails.h"
 #include "UObject/Package.h"
-
-#include "PostProcess/SceneFilterRendering.h" // Renderer/Private
+#include "PostProcess/DrawRectangle.h"
+#include "PixelShaderUtils.h"
 
 #define LOCTEXT_NAMESPACE "FMediaSourceEditorToolkit"
 
@@ -448,14 +448,14 @@ void FMediaSourceEditorToolkit::GenerateThumbnail()
 				SetScreenPassPipelineState(RHICmdList, PipelineState);
 				SetShaderParameters(RHICmdList, PixelShader, PixelShader.GetPixelShader(), *PixelShaderParameters);
 
-				DrawRectangle(
-					RHICmdList,
+				FRHIBatchedShaderParameters& BatchedParameters = RHICmdList.GetScratchShaderParameters();
+				UE::Renderer::PostProcess::SetDrawRectangleParameters(BatchedParameters, PipelineState.VertexShader.GetShader(),
 					0.0f, 0.0f, static_cast<float>(Extent.X), static_cast<float>(Extent.Y),
 					0.0f, 0.0f, static_cast<float>(Extent.X), static_cast<float>(Extent.Y),
 					Extent,
-					Extent,
-					PipelineState.VertexShader,
-					EDRF_UseTriangleOptimization);
+					Extent);
+				RHICmdList.SetBatchedShaderParameters(PipelineState.VertexShader.GetVertexShader(), BatchedParameters);
+				FPixelShaderUtils::DrawFullscreenTriangle(RHICmdList);
 			});
 
 			GraphBuilder.Execute();
