[Media IO] Add Media Player Slider Widget

- Wrapped Playback scrubber in a SMediaPlayerSlider widget + exposed via module interface so other modules can use it

- Media Plate Component customization is now using the SMediaPlayerSlider widget.

diff --git a/Engine/Plugins/Media/MediaPlate/Source/MediaPlateEditor/Private/MediaPlateCustomization.cpp b/Engine/Plugins/Media/MediaPlate/Source/MediaPlateEditor/Private/MediaPlateCustomization.cpp
--- a/Engine/Plugins/Media/MediaPlate/Source/MediaPlateEditor/Private/MediaPlateCustomization.cpp
+++ b/Engine/Plugins/Media/MediaPlate/Source/MediaPlateEditor/Private/MediaPlateCustomization.cpp
@@ -15,6 +15,7 @@
 #include "MediaPlateEditorModule.h"
 #include "MediaPlateEditorStyle.h"
 #include "MediaPlayer.h"
+#include "MediaPlayerEditorModule.h"
 #include "MediaPlaylist.h"
 #include "MediaSource.h"
 #include "PropertyCustomizationHelpers.h"
@@ -181,6 +182,21 @@ void FMediaPlateCustomization::CustomizeDetails(IDetailLayoutBuilder& DetailBuil
 				.OnPathPicked(this, &FMediaPlateCustomization::HandleMediaPathPicked)
 		];
 
+	// Add media player playback slider
+	if (IMediaPlayerEditorModule* MediaPlayerEditorModule = FModuleManager::LoadModulePtr<IMediaPlayerEditorModule>("MediaPlayerEditor"))
+	{
+		const TSharedRef<IMediaPlayerSlider> MediaPlayerSlider =
+			MediaPlayerEditorModule->CreateMediaPlayerSliderWidget(GetMediaPlayer());
+
+		MediaPlayerSlider->SetSliderHandleColor(FSlateColor(EStyleColor::AccentBlue));
+		MediaPlayerSlider->SetVisibleWhenInactive(EVisibility::Visible);
+
+		ControlCategory.AddCustomRow(LOCTEXT("MediaPlatePlaybackPosition", "Playback Position"))
+		[
+			MediaPlayerSlider
+		];
+	}
+
 	// Add media control buttons.
 	ControlCategory.AddCustomRow(LOCTEXT("MediaPlateControls", "MediaPlate Controls"))
 		[
diff --git a/Engine/Plugins/Media/MediaPlayerEditor/Source/MediaPlayerEditor/Private/MediaPlayerEditorModule.cpp b/Engine/Plugins/Media/MediaPlayerEditor/Source/MediaPlayerEditor/Private/MediaPlayerEditorModule.cpp
--- a/Engine/Plugins/Media/MediaPlayerEditor/Source/MediaPlayerEditor/Private/MediaPlayerEditorModule.cpp
+++ b/Engine/Plugins/Media/MediaPlayerEditor/Source/MediaPlayerEditor/Private/MediaPlayerEditorModule.cpp
@@ -45,6 +45,7 @@
 #include "Visualizers/MediaSoundComponentVisualizer.h"
 #include "ToolMenus.h"
 #include "Widgets/SMediaPlayerEditorMedia.h"
+#include "Widgets/SMediaPlayerSlider.h"
 
 #include "MediaPlayerEditorLog.h"
 
@@ -88,6 +89,12 @@ public:
 		return Style;
 	}
 
+	virtual TSharedRef<IMediaPlayerSlider> CreateMediaPlayerSliderWidget(UMediaPlayer* InMediaPlayer, const FSliderStyle& InStyle) override
+	{
+		return SNew(SMediaPlayerSlider, InMediaPlayer)
+		.Style(&InStyle);
+	}
+
 	//~ IModuleInterface interface
 
 	virtual void StartupModule() override
diff --git a/Engine/Plugins/Media/MediaPlayerEditor/Source/MediaPlayerEditor/Private/Widgets/SMediaPlayerEditorViewer.cpp b/Engine/Plugins/Media/MediaPlayerEditor/Source/MediaPlayerEditor/Private/Widgets/SMediaPlayerEditorViewer.cpp
--- a/Engine/Plugins/Media/MediaPlayerEditor/Source/MediaPlayerEditor/Private/Widgets/SMediaPlayerEditorViewer.cpp
+++ b/Engine/Plugins/Media/MediaPlayerEditor/Source/MediaPlayerEditor/Private/Widgets/SMediaPlayerEditorViewer.cpp
@@ -16,6 +16,7 @@
 #include "Framework/Commands/UIAction.h"
 #include "Framework/MultiBox/MultiBoxBuilder.h"
 #include "Framework/Notifications/NotificationManager.h"
+#include "SMediaPlayerSlider.h"
 #include "Styling/CoreStyle.h"
 #include "Textures/SlateIcon.h"
 #include "Widgets/SBoxPanel.h"
@@ -25,7 +26,6 @@
 #include "Widgets/Input/SButton.h"
 #include "Widgets/Input/SComboButton.h"
 #include "Widgets/Input/SEditableTextBox.h"
-#include "Widgets/Input/SSlider.h"
 #include "Widgets/Layout/SBorder.h"
 #include "Widgets/Layout/SBox.h"
 #include "Widgets/Layout/SSpacer.h"
@@ -368,62 +368,15 @@ void SMediaPlayerEditorViewer::Construct(const FArguments& InArgs, UMediaPlayer&
 											[
 												// cache visualization
 												SNew(SMediaPlayerEditorCache, InMediaPlayer, InStyle)
-											]												
+											]
 
 										+ SOverlay::Slot()
 											.VAlign(VAlign_Top)
 											[
 												// time slider
-												SAssignNew(ScrubberSlider, SSlider)
-													.IsEnabled_Lambda([this]() -> bool
-													{
-														return MediaPlayer->SupportsSeeking();
-													})
-													.OnMouseCaptureBegin_Lambda([this]()
-													{
-														ScrubValue = static_cast<float>(FTimespan::Ratio(MediaPlayer->GetDisplayTime(), MediaPlayer->GetDuration()));
-
-														if (MediaPlayer->SupportsScrubbing())
-														{
-															PreScrubRate = MediaPlayer->GetRate();
-															MediaPlayer->SetRate(0.0f);
-														}
-													})
-													.OnMouseCaptureEnd_Lambda([this]()
-													{
-														if (MediaPlayer->SupportsScrubbing())
-														{
-															MediaPlayer->SetRate(PreScrubRate);
-														}
-													})
-													.OnValueChanged_Lambda([this](float NewValue)
-													{
-														ScrubValue = NewValue;
-
-														if (!ScrubberSlider->HasMouseCapture() || MediaPlayer->SupportsScrubbing())
-														{
-															MediaPlayer->Seek(MediaPlayer->GetDuration() * NewValue);
-														}
-													})
-													.Orientation(Orient_Horizontal)
-													.SliderBarColor(FLinearColor::Transparent)
-													.Style(&InStyle->GetWidgetStyle<FSliderStyle>("MediaPlayerEditor.Scrubber"))
-													.Value_Lambda([this]() -> float
-													{
-														if (ScrubberSlider->HasMouseCapture())
-														{
-															return ScrubValue;
-														}
-
-														return (MediaPlayer->GetDuration() > FTimespan::Zero())
-															? static_cast<float>(FTimespan::Ratio(MediaPlayer->GetDisplayTime(), MediaPlayer->GetDuration()))
-															: 0.0f;
-													})
-													.Visibility_Lambda([this]() -> EVisibility {
-														return (MediaPlayer->SupportsScrubbing() || MediaPlayer->SupportsSeeking())
-															? EVisibility::Visible
-															: EVisibility::Hidden;
-													})
+												SNew(SMediaPlayerSlider, MediaPlayer)
+												.ToolTipText( LOCTEXT("PlaybackPosition", "Current Playback Position"))
+												.Style(&InStyle->GetWidgetStyle<FSliderStyle>("MediaPlayerEditor.Scrubber"))
 											]
 
 										+ SOverlay::Slot()
diff --git a/Engine/Plugins/Media/MediaPlayerEditor/Source/MediaPlayerEditor/Private/Widgets/SMediaPlayerSlider.cpp b/Engine/Plugins/Media/MediaPlayerEditor/Source/MediaPlayerEditor/Private/Widgets/SMediaPlayerSlider.cpp
new file mode 100644
--- /dev/null
+++ b/Engine/Plugins/Media/MediaPlayerEditor/Source/MediaPlayerEditor/Private/Widgets/SMediaPlayerSlider.cpp
@@ -0,0 +1,128 @@
+// Copyright Epic Games, Inc. All Rights Reserved.
+
+#include "SMediaPlayerSlider.h"
+#include "MediaPlayer.h"
+#include "SlateOptMacros.h"
+#include "Widgets/Input/SSlider.h"
+
+BEGIN_SLATE_FUNCTION_BUILD_OPTIMIZATION
+void SMediaPlayerSlider::Construct(const FArguments& InArgs, UMediaPlayer* InMediaPlayer)
+{
+	MediaPlayerWeak = InMediaPlayer;
+
+	ChildSlot
+	[
+		SAssignNew(ScrubberSlider, SSlider)
+		.IsEnabled_Raw(this, &SMediaPlayerSlider::DoesMediaPlayerSupportSeeking)
+		.OnMouseCaptureBegin_Raw(this, &SMediaPlayerSlider::OnScrubBegin)
+		.OnMouseCaptureEnd_Raw(this, &SMediaPlayerSlider::OnScrubEnd)
+		.OnValueChanged_Raw(this, &SMediaPlayerSlider::Seek)
+		.Value_Raw(this, &SMediaPlayerSlider::GetPlaybackPosition)
+		.Visibility_Raw(this, &SMediaPlayerSlider::GetScrubberVisibility)
+		.Orientation(Orient_Horizontal)
+		.SliderBarColor(FLinearColor::Transparent)
+		.Style(InArgs._Style)
+	];
+}
+END_SLATE_FUNCTION_BUILD_OPTIMIZATION
+
+bool SMediaPlayerSlider::DoesMediaPlayerSupportSeeking() const
+{
+	if (const UMediaPlayer* MediaPlayer = MediaPlayerWeak.Get())
+	{
+		return MediaPlayer->SupportsSeeking();
+	}
+
+	return false;
+}
+
+void SMediaPlayerSlider::OnScrubBegin()
+{
+	if (UMediaPlayer* MediaPlayer = MediaPlayerWeak.Get())
+	{
+		ScrubValue = static_cast<float>(FTimespan::Ratio(MediaPlayer->GetDisplayTime(), MediaPlayer->GetDuration()));
+
+		if (MediaPlayer->SupportsScrubbing())
+		{
+			PreScrubRate = MediaPlayer->GetRate();
+			MediaPlayer->SetRate(0.0f);
+		}
+	}
+}
+
+void SMediaPlayerSlider::OnScrubEnd()
+{
+	if (UMediaPlayer* MediaPlayer = MediaPlayerWeak.Get())
+	{
+		if (MediaPlayer->SupportsScrubbing())
+		{
+			MediaPlayer->SetRate(PreScrubRate);
+		}
+
+		// Set playback position to scrub value when drag ends
+		MediaPlayer->Seek(MediaPlayer->GetDuration() * ScrubValue);
+	}
+}
+
+void SMediaPlayerSlider::Seek(float InPlaybackPosition)
+{
+	if (UMediaPlayer* MediaPlayer = MediaPlayerWeak.Get())
+	{
+		ScrubValue = InPlaybackPosition;
+
+		if (!ScrubberSlider->HasMouseCapture() || MediaPlayer->SupportsScrubbing())
+		{
+			MediaPlayer->Seek(MediaPlayer->GetDuration() * InPlaybackPosition);
+		}
+	}
+}
+
+float SMediaPlayerSlider::GetPlaybackPosition() const
+{
+	if (const UMediaPlayer* MediaPlayer = MediaPlayerWeak.Get())
+	{
+		if (ScrubberSlider->HasMouseCapture())
+		{
+			return ScrubValue;
+		}
+
+		return (MediaPlayer->GetDuration() > FTimespan::Zero())
+			? static_cast<float>(FTimespan::Ratio(MediaPlayer->GetDisplayTime(), MediaPlayer->GetDuration()))
+			: 0.0f;
+	}
+
+	return 0.0f;
+}
+
+EVisibility SMediaPlayerSlider::GetScrubberVisibility() const
+{
+	if (const UMediaPlayer* MediaPlayer = MediaPlayerWeak.Get())
+	{
+		return (MediaPlayer->SupportsScrubbing() || MediaPlayer->SupportsSeeking())
+															? EVisibility::Visible
+															: VisibilityWhenInactive;
+	}
+
+	return VisibilityWhenInactive;
+}
+
+void SMediaPlayerSlider::SetSliderHandleColor(const FSlateColor& InSliderColor)
+{
+	if (ScrubberSlider)
+	{
+		ScrubberSlider->SetSliderHandleColor(InSliderColor);
+	}
+}
+
+void SMediaPlayerSlider::SetSliderBarColor(const FSlateColor& InSliderColor)
+{
+	if (ScrubberSlider)
+	{
+		ScrubberSlider->SetSliderBarColor(InSliderColor);
+	}
+}
+
+void SMediaPlayerSlider::SetVisibleWhenInactive(EVisibility InVisibility)
+{
+	VisibilityWhenInactive = InVisibility;
+}
diff --git a/Engine/Plugins/Media/MediaPlayerEditor/Source/MediaPlayerEditor/Private/Widgets/SMediaPlayerSlider.h b/Engine/Plugins/Media/MediaPlayerEditor/Source/MediaPlayerEditor/Private/Widgets/SMediaPlayerSlider.h
new file mode 100644
--- /dev/null
+++ b/Engine/Plugins/Media/MediaPlayerEditor/Source/MediaPlayerEditor/Private/Widgets/SMediaPlayerSlider.h
@@ -0,0 +1,61 @@
+// Copyright Epic Games, Inc. All Rights Reserved.
+
+#pragma once
+
+#include "MediaPlayerEditorModule.h"
+#include "Styling/CoreStyle.h"
+#include "Styling/SlateTypes.h"
+#include "Styling/SlateWidgetStyleAsset.h"
+#include "Widgets/SCompoundWidget.h"
+
+class SSlider;
+class UMediaPlayer;
+
+/**
+ * Implements a scrubber to visualize the current playback position of a Media Player
+ * and interact with it.
+ */
+class SMediaPlayerSlider : public IMediaPlayerSlider
+{
+public:
+	SLATE_BEGIN_ARGS(SMediaPlayerSlider)
+			: _Style(&FCoreStyle::Get().GetWidgetStyle<FSliderStyle>("Slider"))
+		{
+		}
+
+	/** The Slider style used to draw the scrubber. */
+	SLATE_STYLE_ARGUMENT(FSliderStyle, Style)
+	SLATE_END_ARGS()
+
+	/** Constructs this widget with InArgs */
+	void Construct(const FArguments& InArgs, UMediaPlayer* InMediaPlayer);
+
+	//~Begin IMediaPlayerScrubber
+	virtual void SetSliderHandleColor(const FSlateColor& InSliderColor) override;
+	virtual void SetSliderBarColor(const FSlateColor& InSliderColor) override;
+	virtual void SetVisibleWhenInactive(EVisibility InVisibility) override;
+	//~End IMediaPlayerScrubber
+
+private:
+	bool DoesMediaPlayerSupportSeeking() const;
+	void OnScrubBegin();
+	void OnScrubEnd();
+	void Seek(float InPlaybackPosition);
+	float GetPlaybackPosition() const;
+	EVisibility GetScrubberVisibility() const;
+
+	/** The scrubber visibility when inactive. */
+	EVisibility VisibilityWhenInactive = EVisibility::Hidden;
+
+	/** Pointer to the media player that is being viewed. */
+	TWeakObjectPtr<UMediaPlayer> MediaPlayerWeak;
+
+	/** The playback rate prior to scrubbing. */
+	float PreScrubRate = 0.0f;
+
+	/** Holds the scrubber slider. */
+	TSharedPtr<SSlider> ScrubberSlider;
+
+	/** The value currently being scrubbed to. */
+	float ScrubValue = 0.0f;
+};
diff --git a/Engine/Plugins/Media/MediaPlayerEditor/Source/MediaPlayerEditor/Public/MediaPlayerEditorModule.h b/Engine/Plugins/Media/MediaPlayerEditor/Source/MediaPlayerEditor/Public/MediaPlayerEditorModule.h
--- a/Engine/Plugins/Media/MediaPlayerEditor/Source/MediaPlayerEditor/Public/MediaPlayerEditorModule.h
+++ b/Engine/Plugins/Media/MediaPlayerEditor/Source/MediaPlayerEditor/Public/MediaPlayerEditorModule.h
@@ -3,10 +3,30 @@
 #pragma once
 
 #include "Modules/ModuleInterface.h"
+#include "Styling/CoreStyle.h"
+#include "Styling/SlateTypes.h"
 #include "Templates/SharedPointer.h"
+#include "Widgets/SCompoundWidget.h"
 
 class FImgMediaPlayer;
 class ISlateStyle;
+class UMediaPlayer;
+
+/**
+ * Interface for MediaPlayer Playback Scrubber Widget
+ */
+class MEDIAPLAYEREDITOR_API IMediaPlayerSlider : public SCompoundWidget
+{
+public:
+	/** Set the Scrubber Slider Handle Color */
+	virtual void SetSliderHandleColor(const FSlateColor& InSliderColor) = 0;
+
+	/** Set the Scrubber Slider Color */
+	virtual void SetSliderBarColor(const FSlateColor& InSliderColor) = 0;
+
+	/** Set the Scrubber Slider Visibility, when the player is inactive */
+	virtual void SetVisibleWhenInactive(EVisibility InVisibility) = 0;
+};
 
 /**
 * Interface for the MediaPlayerEditor module.
@@ -15,8 +35,15 @@ class MEDIAPLAYEREDITOR_API IMediaPlayerEditorModule
 	: public IModuleInterface
 {
 public:
-	
+
 	/** Get the style used by this module. */
 	virtual TSharedPtr<ISlateStyle> GetStyle() = 0;
 
+	/**
+	 * Creates a Widget to visualize playback time and scrub the content played by a Media Player
+	 * @param InMediaPlayer: the player affected by the widget
+	 * @param InStyle: the style chosen for this slider widget
+	 * @return Scrubber Widget
+	 */
+	virtual TSharedRef<IMediaPlayerSlider> CreateMediaPlayerSliderWidget(UMediaPlayer* InMediaPlayer, const FSliderStyle& InStyle = FCoreStyle::Get().GetWidgetStyle<FSliderStyle>("Slider")) = 0;
 };
diff --git a/Engine/Plugins/Media/MediaPlayerEditor/Source/MediaPlayerEditor/Public/Widgets/SMediaPlayerEditorViewer.h b/Engine/Plugins/Media/MediaPlayerEditor/Source/MediaPlayerEditor/Public/Widgets/SMediaPlayerEditorViewer.h
--- a/Engine/Plugins/Media/MediaPlayerEditor/Source/MediaPlayerEditor/Public/Widgets/SMediaPlayerEditorViewer.h
+++ b/Engine/Plugins/Media/MediaPlayerEditor/Source/MediaPlayerEditor/Public/Widgets/SMediaPlayerEditorViewer.h
@@ -13,7 +13,6 @@
 
 class FMenuBuilder;
 class SEditableTextBox;
-class SSlider;
 class SMediaPlayerEditorViewport;
 class UMediaPlayer;
 class UMediaSoundComponent;
@@ -142,15 +141,6 @@ private:
 	/** Pointer to the media player that is being viewed. */
 	UMediaPlayer* MediaPlayer;
 
-	/** The playback rate prior to scrubbing. */
-	float PreScrubRate;
-
-	/** Holds the scrubber slider. */
-	TSharedPtr<SSlider> ScrubberSlider;
-
-	/** The value currently being scrubbed to. */
-	float ScrubValue;
-
 	/** The style set to use for this widget. */
 	TSharedPtr<ISlateStyle> Style;
 
