MediaPlayerEditor: Fixed thumbnail creation to only do this when the clip contains video and is of a supported format. Fixed state tracking as media players may not be in the required states immediately after requesting them.

#rnx

diff --git a/Engine/Plugins/Media/MediaPlayerEditor/Source/MediaPlayerEditor/Private/Shared/MediaSourceRenderer.cpp b/Engine/Plugins/Media/MediaPlayerEditor/Source/MediaPlayerEditor/Private/Shared/MediaSourceRenderer.cpp
--- a/Engine/Plugins/Media/MediaPlayerEditor/Source/MediaPlayerEditor/Private/Shared/MediaSourceRenderer.cpp
+++ b/Engine/Plugins/Media/MediaPlayerEditor/Source/MediaPlayerEditor/Private/Shared/MediaSourceRenderer.cpp
@@ -15,7 +15,8 @@ UMediaTexture* UMediaSourceRenderer::Open(UMediaSource* InMediaSource)
 		if (MediaPlayer == nullptr)
 		{
 			MediaPlayer = NewObject<UMediaPlayer>(GetTransientPackage());
-			MediaPlayer->OnSeekCompleted.AddDynamic(this, &UMediaSourceRenderer::OnSeekCompleted);
+			MediaPlayer->OnMediaOpened.AddDynamic(this, &UMediaSourceRenderer::OnMediaOpened);
+			MediaPlayer->OnMediaOpenFailed.AddDynamic(this, &UMediaSourceRenderer::OnMediaOpenFailed);
 		}
 		else
 		{
@@ -34,10 +35,16 @@ UMediaTexture* UMediaSourceRenderer::Open(UMediaSource* InMediaSource)
 		MediaSource = InMediaSource;
 
 		// Start playing the media.
-		bIsSeekActive = false;
+		CurrentState = EState::Opening;
+		WatchdogTimeRemaining = 10.0f;
+
 		FMediaPlayerOptions Options;
-		Options.PlayOnOpen = EMediaPlayerOptionBooleanOverride::Enabled;
+		Options.PlayOnOpen = EMediaPlayerOptionBooleanOverride::Disabled;
 		Options.Loop = EMediaPlayerOptionBooleanOverride::Disabled;
+		// Let the media start at whichever time it defaults to.
+		Options.SeekTimeType = EMediaPlayerOptionSeekTimeType::Ignored;
+		// We don't need audio.
+		Options.Tracks.Audio = -1;
 		// For image media, we avoid filling the global cache which will needlessly hold onto frame data.
 		Options.InternalCustomOptions.Emplace(MediaPlayerOptionValues::ImgMediaSmartCacheEnabled(), FVariant(true));
 		Options.InternalCustomOptions.Emplace(MediaPlayerOptionValues::ImgMediaSmartCacheTimeToLookAhead(), FVariant(0.2f));
@@ -55,6 +62,11 @@ void UMediaSourceRenderer::Tick(float DeltaTime)
 {
 	if (MediaPlayer != nullptr)
 	{
+		// Keep the delta time in check. We typically create a thumbnail after having selected a new source
+		// and if that brought up the system file selector box the time spent in the file selector is included
+		// in the delta time.
+		DeltaTime = FMath::Min(DeltaTime, 0.1f);
+
 		// Is the texture ready?
 		// The aspect ratio will change when we have something.
 		if (MediaTexture->CurrentAspectRatio != 0.0f)
@@ -70,37 +82,102 @@ void UMediaSourceRenderer::Tick(float DeltaTime)
 		}
 		else
 		{
-			if (MediaPlayer->GetUrl().IsEmpty())
+			switch(CurrentState)
 			{
-				Close();
-			}
-			else if (MediaPlayer->IsPreparing() == false)
-			{
-				// Don't advance playback.
-				if (MediaPlayer->IsPlaying())
+				case EState::Closed:
+				case EState::Errored:
+				{
+					break;
+				}
+				case EState::Opening:
 				{
-					MediaPlayer->SetRate(0.0f);
+					// Make sure this doesn't drag on forever.
+					if ((WatchdogTimeRemaining -= DeltaTime) < 0.0f)
+					{
+						CurrentState = EState::TimedOut;
+					}
+					break;
 				}
+				case EState::Open:
+				{
+					// Make sure this doesn't drag on forever.
+					if ((WatchdogTimeRemaining -= DeltaTime) < 0.0f)
+					{
+						CurrentState = EState::TimedOut;
+						return;
+					}
+
+					// Check that we have a media duration. When we have that we can assume
+					// that we also know which tracks are present.
+					const FTimespan MediaDuration = MediaPlayer->GetDuration();
+					if (MediaDuration > FTimespan::Zero() && MediaDuration < FTimespan::MaxValue())
+					{
+						// Check if there is a video track.
+						if (MediaPlayer->GetNumTracks(EMediaPlayerTrack::Video) <= 0 || MediaPlayer->GetTrackFormat(EMediaPlayerTrack::Video, 0) < 0)
+						{
+							CurrentState = EState::NotSupported;
+						}
+						else
+						{
+							// Check if the media can be seeked as we don't want the first frame
+							// which is often just a black frame.
+							if (MediaPlayer->SupportsSeeking())
+							{
+								const float kSeekPosAsDurationScale = 0.3f;
 
-				// If we can seek, then seek.
-				if ((MediaPlayer->IsBuffering() == false) && (MediaPlayer->IsReady()))
+								FTimespan SeekTime = MediaDuration * kSeekPosAsDurationScale;
+								// If seeking is supported, check if timeranges are as well.
+								// If they are, chances are the media does not start at zero.
+								if (MediaPlayer->SupportsPlaybackTimeRange())
+								{
+									TRange<FTimespan> Range = MediaPlayer->GetPlaybackTimeRange(EMediaTimeRangeType::Current);
+									SeekTime = Range.GetLowerBoundValue();
+									if (Range.GetUpperBoundValue() < FTimespan::MaxValue())
+									{
+										SeekTime += (Range.GetUpperBoundValue() - Range.GetLowerBoundValue()) * kSeekPosAsDurationScale;
+									}
+								}
+								MediaPlayer->Seek(SeekTime);
+							}
+							MediaPlayer->Play();
+							CurrentState = EState::Playing;
+							WatchdogTimeRemaining = 5.0f;
+						}
+					}
+					break;
+				}
+				case EState::Playing:
 				{
-					if (bIsSeekActive == false)
+					if ((WatchdogTimeRemaining -= DeltaTime) < 0.0f)
 					{
-						bIsSeekActive = true;
-						const FTimespan MediaDuration = MediaPlayer->GetDuration();
-						FTimespan SeekTime = MediaDuration * 0.3f;
-						MediaPlayer->Seek(SeekTime);
+						CurrentState = EState::TimedOut;
 					}
+					break;
+				}
+				case EState::TimedOut:
+				case EState::NotSupported:
+				case EState::Failed:
+				{
+					// Opening failed. No point in doing anything, so close everything.
+					CurrentState = EState::Errored;
+					Close();
+					return;
 				}
 			}
 		}
 	}
 }
 
-void UMediaSourceRenderer::OnSeekCompleted()
+void UMediaSourceRenderer::OnMediaOpened(FString)
+{
+	// Give it a moment in case the metadata is not immediately available.
+	WatchdogTimeRemaining = 1.0f;
+	CurrentState = EState::Open;
+}
+
+void UMediaSourceRenderer::OnMediaOpenFailed(FString)
 {
-	bIsSeekActive = false;
+	CurrentState = EState::Failed;
 }
 
 void UMediaSourceRenderer::Close()
@@ -114,4 +191,6 @@ void UMediaSourceRenderer::Close()
 		MediaPlayer->Close();
 		MediaPlayer = nullptr;
 	}
+	CurrentState = EState::Closed;
+	WatchdogTimeRemaining = 0.0f;
 }
diff --git a/Engine/Plugins/Media/MediaPlayerEditor/Source/MediaPlayerEditor/Private/Shared/MediaSourceRenderer.h b/Engine/Plugins/Media/MediaPlayerEditor/Source/MediaPlayerEditor/Private/Shared/MediaSourceRenderer.h
--- a/Engine/Plugins/Media/MediaPlayerEditor/Source/MediaPlayerEditor/Private/Shared/MediaSourceRenderer.h
+++ b/Engine/Plugins/Media/MediaPlayerEditor/Source/MediaPlayerEditor/Private/Shared/MediaSourceRenderer.h
@@ -29,12 +29,6 @@ public:
 	virtual TStatId GetStatId() const override { RETURN_QUICK_DECLARE_CYCLE_STAT(UMediaSourceRenderer, STATGROUP_Tickables); }
 
 private:
-	/**
-	 * Callback when the media player is done seeking.
-	 */
-	UFUNCTION()
-	void OnSeekCompleted();
-
 	/**
 	 * Cleans everything up.
 	 * The media textue will remain so it can be used/reused.
@@ -53,7 +47,21 @@ private:
 	UPROPERTY(Transient)
 	TObjectPtr<UMediaTexture> MediaTexture = nullptr;
 
-	/** True if we are currently seeking. */
-	bool bIsSeekActive = false;
-	
+	enum class EState
+	{
+		Closed,
+		Opening,
+		Open,
+		Playing,
+		NotSupported,
+		Failed,
+		TimedOut,
+		Errored
+	};
+	EState CurrentState = EState::Closed;
+	float WatchdogTimeRemaining = 0.0f;
+	UFUNCTION()
+	void OnMediaOpened(FString InURL);
+	UFUNCTION()
+	void OnMediaOpenFailed(FString InURL);
 };
