Making the MediaPlayer GUI slider logic a bit more "friendly" in using the media player API for seeking

diff --git a/Engine/Plugins/Media/MediaPlayerEditor/Source/MediaPlayerEditor/Private/Widgets/SMediaPlayerSlider.cpp b/Engine/Plugins/Media/MediaPlayerEditor/Source/MediaPlayerEditor/Private/Widgets/SMediaPlayerSlider.cpp
--- a/Engine/Plugins/Media/MediaPlayerEditor/Source/MediaPlayerEditor/Private/Widgets/SMediaPlayerSlider.cpp
+++ b/Engine/Plugins/Media/MediaPlayerEditor/Source/MediaPlayerEditor/Private/Widgets/SMediaPlayerSlider.cpp
@@ -41,6 +41,7 @@ void SMediaPlayerSlider::OnScrubBegin()
 	if (UMediaPlayer* MediaPlayer = MediaPlayerWeak.Get())
 	{
 		ScrubValue = static_cast<float>(FTimespan::Ratio(MediaPlayer->GetDisplayTime(), MediaPlayer->GetDuration()));
+		LastScrubValue = ScrubValue;
 
 		if (MediaPlayer->SupportsScrubbing())
 		{
@@ -54,13 +55,16 @@ void SMediaPlayerSlider::OnScrubEnd()
 {
 	if (UMediaPlayer* MediaPlayer = MediaPlayerWeak.Get())
 	{
+		// Set playback position to scrub value when drag ends
+		if (LastScrubValue != ScrubValue)
+		{
+			MediaPlayer->Seek(MediaPlayer->GetDuration() * ScrubValue);
+		}
+
 		if (MediaPlayer->SupportsScrubbing())
 		{
 			MediaPlayer->SetRate(PreScrubRate);
 		}
-
-		// Set playback position to scrub value when drag ends
-		MediaPlayer->Seek(MediaPlayer->GetDuration() * ScrubValue);
 	}
 }
 
@@ -73,6 +77,7 @@ void SMediaPlayerSlider::Seek(float InPlaybackPosition)
 		if (!ScrubberSlider->HasMouseCapture() || MediaPlayer->SupportsScrubbing())
 		{
 			MediaPlayer->Seek(MediaPlayer->GetDuration() * InPlaybackPosition);
+			LastScrubValue = ScrubValue;
 		}
 	}
 }
diff --git a/Engine/Plugins/Media/MediaPlayerEditor/Source/MediaPlayerEditor/Private/Widgets/SMediaPlayerSlider.h b/Engine/Plugins/Media/MediaPlayerEditor/Source/MediaPlayerEditor/Private/Widgets/SMediaPlayerSlider.h
--- a/Engine/Plugins/Media/MediaPlayerEditor/Source/MediaPlayerEditor/Private/Widgets/SMediaPlayerSlider.h
+++ b/Engine/Plugins/Media/MediaPlayerEditor/Source/MediaPlayerEditor/Private/Widgets/SMediaPlayerSlider.h
@@ -58,4 +58,7 @@ private:
 
 	/** The value currently being scrubbed to. */
 	float ScrubValue = 0.0f;
+
+	/** The last value set with media player while scrubbing. */
+	float LastScrubValue = -1.0f;
 };
