[MediaPlate] MediaPlateCustomization: Seek Slider supports multiple selection.

diff --git a/Engine/Plugins/Media/MediaPlate/Source/MediaPlateEditor/Private/MediaPlateCustomization.cpp b/Engine/Plugins/Media/MediaPlate/Source/MediaPlateEditor/Private/MediaPlateCustomization.cpp
--- a/Engine/Plugins/Media/MediaPlate/Source/MediaPlateEditor/Private/MediaPlateCustomization.cpp
+++ b/Engine/Plugins/Media/MediaPlate/Source/MediaPlateEditor/Private/MediaPlateCustomization.cpp
@@ -94,7 +94,7 @@ void FMediaPlateCustomization::CustomizeDetails(IDetailLayoutBuilder& DetailBuil
 	if (IMediaPlayerEditorModule* MediaPlayerEditorModule = FModuleManager::LoadModulePtr<IMediaPlayerEditorModule>("MediaPlayerEditor"))
 	{
 		const TSharedRef<IMediaPlayerSlider> MediaPlayerSlider =
-			MediaPlayerEditorModule->CreateMediaPlayerSliderWidget(GetMediaPlayer());
+			MediaPlayerEditorModule->CreateMediaPlayerSliderWidget(GetMediaPlayers());
 
 		MediaPlayerSlider->SetSliderHandleColor(FSlateColor(EStyleColor::AccentBlue));
 		MediaPlayerSlider->SetVisibleWhenInactive(EVisibility::Visible);
@@ -1012,6 +1012,24 @@ UMediaPlayer* FMediaPlateCustomization::GetMediaPlayer() const
 	return nullptr;
 }
 
+TArray<TWeakObjectPtr<UMediaPlayer>> FMediaPlateCustomization::GetMediaPlayers() const
+{
+	TArray<TWeakObjectPtr<UMediaPlayer>> MediaPlayers;
+	MediaPlayers.Reserve(MediaPlatesList.Num());
+	
+	for (const TWeakObjectPtr<UMediaPlateComponent>& MediaPlateWeak : MediaPlatesList)
+	{
+		if (UMediaPlateComponent* MediaPlate = MediaPlateWeak.Get())
+		{
+			if (UMediaPlayer* MediaPlayer = MediaPlate->GetMediaPlayer())
+			{
+				MediaPlayers.Add(MediaPlayer);
+			}
+		}
+	}
+	return MediaPlayers;
+}
+
 void FMediaPlateCustomization::CustomizeCategories(IDetailLayoutBuilder& InDetailBuilder)
 {
 	static const FName PropertyEditor("PropertyEditor");
diff --git a/Engine/Plugins/Media/MediaPlate/Source/MediaPlateEditor/Private/MediaPlateCustomization.h b/Engine/Plugins/Media/MediaPlate/Source/MediaPlateEditor/Private/MediaPlateCustomization.h
--- a/Engine/Plugins/Media/MediaPlate/Source/MediaPlateEditor/Private/MediaPlateCustomization.h
+++ b/Engine/Plugins/Media/MediaPlate/Source/MediaPlateEditor/Private/MediaPlateCustomization.h
@@ -186,4 +186,6 @@ private:
 	void StopMediaPlates();
 
 	UMediaPlayer* GetMediaPlayer() const;
+
+	TArray<TWeakObjectPtr<UMediaPlayer>> GetMediaPlayers() const;
 };
diff --git a/Engine/Plugins/Media/MediaPlayerEditor/Source/MediaPlayerEditor/Private/MediaPlayerEditorModule.cpp b/Engine/Plugins/Media/MediaPlayerEditor/Source/MediaPlayerEditor/Private/MediaPlayerEditorModule.cpp
--- a/Engine/Plugins/Media/MediaPlayerEditor/Source/MediaPlayerEditor/Private/MediaPlayerEditorModule.cpp
+++ b/Engine/Plugins/Media/MediaPlayerEditor/Source/MediaPlayerEditor/Private/MediaPlayerEditorModule.cpp
@@ -54,7 +54,6 @@
 
 DEFINE_LOG_CATEGORY(LogMediaPlayerEditor);
 
-
 /**
  * Implements the MediaPlayerEditor module.
  */
@@ -89,9 +88,15 @@ public:
 		return Style;
 	}
 
-	virtual TSharedRef<IMediaPlayerSlider> CreateMediaPlayerSliderWidget(UMediaPlayer* InMediaPlayer, const FSliderStyle& InStyle) override
+	TSharedRef<IMediaPlayerSlider> CreateMediaPlayerSliderWidget(UMediaPlayer* InMediaPlayer, const FSliderStyle& InStyle)
+	{
+		TWeakObjectPtr<UMediaPlayer> MediaPlayerWeak = InMediaPlayer;
+		return CreateMediaPlayerSliderWidget(MakeArrayView(&MediaPlayerWeak, 1), InStyle);
+	}
+
+	virtual TSharedRef<IMediaPlayerSlider> CreateMediaPlayerSliderWidget(const TArrayView<TWeakObjectPtr<UMediaPlayer>> InMediaPlayers, const FSliderStyle& InStyle) override
 	{
-		return SNew(SMediaPlayerSlider, InMediaPlayer)
+		return SNew(SMediaPlayerSlider, InMediaPlayers)
 		.Style(&InStyle);
 	}
 
diff --git a/Engine/Plugins/Media/MediaPlayerEditor/Source/MediaPlayerEditor/Private/Widgets/SMediaPlayerEditorViewer.cpp b/Engine/Plugins/Media/MediaPlayerEditor/Source/MediaPlayerEditor/Private/Widgets/SMediaPlayerEditorViewer.cpp
--- a/Engine/Plugins/Media/MediaPlayerEditor/Source/MediaPlayerEditor/Private/Widgets/SMediaPlayerEditorViewer.cpp
+++ b/Engine/Plugins/Media/MediaPlayerEditor/Source/MediaPlayerEditor/Private/Widgets/SMediaPlayerEditorViewer.cpp
@@ -74,6 +74,8 @@ void SMediaPlayerEditorViewer::Construct(const FArguments& InArgs, UMediaPlayer&
 	MediaPlayer->OnMediaEvent().AddSP(this, &SMediaPlayerEditorViewer::HandleMediaPlayerMediaEvent);
 	MediaPlayer->SetDesiredPlayerName(NAME_None);
 
+	TWeakObjectPtr<UMediaPlayer> MediaPlayerWeak(&InMediaPlayer);
+	
 	FName DesiredPlayerName = GetDefault<UMediaPlayerEditorSettings>()->DesiredPlayerName;
 
 	if (DesiredPlayerName != NAME_None)
@@ -374,7 +376,7 @@ void SMediaPlayerEditorViewer::Construct(const FArguments& InArgs, UMediaPlayer&
 											.VAlign(VAlign_Top)
 											[
 												// time slider
-												SNew(SMediaPlayerSlider, MediaPlayer)
+												SNew(SMediaPlayerSlider, MakeArrayView(&MediaPlayerWeak, 1))
 												.ToolTipText( LOCTEXT("PlaybackPosition", "Current Playback Position"))
 												.Style(&InStyle->GetWidgetStyle<FSliderStyle>("MediaPlayerEditor.Scrubber"))
 											]
diff --git a/Engine/Plugins/Media/MediaPlayerEditor/Source/MediaPlayerEditor/Private/Widgets/SMediaPlayerSlider.cpp b/Engine/Plugins/Media/MediaPlayerEditor/Source/MediaPlayerEditor/Private/Widgets/SMediaPlayerSlider.cpp
--- a/Engine/Plugins/Media/MediaPlayerEditor/Source/MediaPlayerEditor/Private/Widgets/SMediaPlayerSlider.cpp
+++ b/Engine/Plugins/Media/MediaPlayerEditor/Source/MediaPlayerEditor/Private/Widgets/SMediaPlayerSlider.cpp
@@ -6,9 +6,16 @@
 #include "Widgets/Input/SSlider.h"
 
 BEGIN_SLATE_FUNCTION_BUILD_OPTIMIZATION
-void SMediaPlayerSlider::Construct(const FArguments& InArgs, UMediaPlayer* InMediaPlayer)
+void SMediaPlayerSlider::Construct(const FArguments& InArgs, const TArrayView<TWeakObjectPtr<UMediaPlayer>> InMediaPlayers)
 {
-	MediaPlayerWeak = InMediaPlayer;
+	MediaPlayerEntries.Reserve(InMediaPlayers.Num());
+	for (const TWeakObjectPtr<UMediaPlayer>& MediaPlayerWeak : InMediaPlayers)
+	{
+		if (MediaPlayerWeak.IsValid())
+		{
+			MediaPlayerEntries.Emplace(MediaPlayerWeak);
+		}
+	}
 
 	ChildSlot
 	[
@@ -28,72 +35,94 @@ END_SLATE_FUNCTION_BUILD_OPTIMIZATION
 
 bool SMediaPlayerSlider::DoesMediaPlayerSupportSeeking() const
 {
-	if (const UMediaPlayer* MediaPlayer = MediaPlayerWeak.Get())
+	// All players must support seek.
+	bool bSupportSeek = false;
+	for (const FMediaPlayerEntry& Entry : MediaPlayerEntries)
 	{
-		return MediaPlayer->SupportsSeeking();
+		if (const UMediaPlayer* MediaPlayer = Entry.MediaPlayerWeak.Get())
+		{
+			bSupportSeek = MediaPlayer->SupportsSeeking();
+			if (!bSupportSeek)
+			{
+				return false;
+			}
+		}
 	}
 
-	return false;
+	return bSupportSeek;
 }
 
 void SMediaPlayerSlider::OnScrubBegin()
 {
-	if (UMediaPlayer* MediaPlayer = MediaPlayerWeak.Get())
+	for (FMediaPlayerEntry& Entry : MediaPlayerEntries)
 	{
-		ScrubValue = static_cast<float>(FTimespan::Ratio(MediaPlayer->GetDisplayTime(), MediaPlayer->GetDuration()));
-		LastScrubValue = ScrubValue;
-
-		if (MediaPlayer->SupportsScrubbing())
+		if (UMediaPlayer* MediaPlayer = Entry.MediaPlayerWeak.Get())
 		{
-			PreScrubRate = MediaPlayer->GetRate();
-			MediaPlayer->SetRate(0.0f);
+			Entry.ScrubValue = static_cast<float>(FTimespan::Ratio(MediaPlayer->GetDisplayTime(), MediaPlayer->GetDuration()));
+			Entry.LastScrubValue = Entry.ScrubValue;
+
+			if (MediaPlayer->SupportsScrubbing())
+			{
+				Entry.PreScrubRate = MediaPlayer->GetRate();
+				MediaPlayer->SetRate(0.0f);
+			}
 		}
 	}
 }
 
 void SMediaPlayerSlider::OnScrubEnd()
 {
-	if (UMediaPlayer* MediaPlayer = MediaPlayerWeak.Get())
+	for (FMediaPlayerEntry& Entry : MediaPlayerEntries)
 	{
-		// Set playback position to scrub value when drag ends
-		if (LastScrubValue != ScrubValue)
+		if (UMediaPlayer* MediaPlayer = Entry.MediaPlayerWeak.Get())
 		{
-			MediaPlayer->Seek(MediaPlayer->GetDuration() * ScrubValue);
-		}
-
-		if (MediaPlayer->SupportsScrubbing())
-		{
-			MediaPlayer->SetRate(PreScrubRate);
+			// Set playback position to scrub value when drag ends
+			if (Entry.LastScrubValue != Entry.ScrubValue)
+			{
+				MediaPlayer->Seek(MediaPlayer->GetDuration() * Entry.ScrubValue);
+			}
+
+			if (MediaPlayer->SupportsScrubbing())
+			{
+				MediaPlayer->SetRate(Entry.PreScrubRate);
+			}
 		}
 	}
 }
 
 void SMediaPlayerSlider::Seek(float InPlaybackPosition)
 {
-	if (UMediaPlayer* MediaPlayer = MediaPlayerWeak.Get())
+	for (FMediaPlayerEntry& Entry : MediaPlayerEntries)
 	{
-		ScrubValue = InPlaybackPosition;
-
-		if (!ScrubberSlider->HasMouseCapture() || MediaPlayer->SupportsScrubbing())
+		if (UMediaPlayer* MediaPlayer = Entry.MediaPlayerWeak.Get())
 		{
-			MediaPlayer->Seek(MediaPlayer->GetDuration() * InPlaybackPosition);
-			LastScrubValue = ScrubValue;
+			Entry.ScrubValue = InPlaybackPosition;
+
+			if (!ScrubberSlider->HasMouseCapture() || MediaPlayer->SupportsScrubbing())
+			{
+				MediaPlayer->Seek(MediaPlayer->GetDuration() * InPlaybackPosition);
+				Entry.LastScrubValue = Entry.ScrubValue;
+			}
 		}
 	}
 }
 
 float SMediaPlayerSlider::GetPlaybackPosition() const
 {
-	if (const UMediaPlayer* MediaPlayer = MediaPlayerWeak.Get())
+	// Note: returns first one. All scrub positions should match.
+	for (const FMediaPlayerEntry& Entry : MediaPlayerEntries)
 	{
-		if (ScrubberSlider->HasMouseCapture())
+		if (const UMediaPlayer* MediaPlayer = Entry.MediaPlayerWeak.Get())
 		{
-			return ScrubValue;
+			if (ScrubberSlider->HasMouseCapture())
+			{
+				return Entry.ScrubValue;
+			}
+
+			return (MediaPlayer->GetDuration() > FTimespan::Zero())
+				? static_cast<float>(FTimespan::Ratio(MediaPlayer->GetDisplayTime(), MediaPlayer->GetDuration()))
+				: 0.0f;
 		}
-
-		return (MediaPlayer->GetDuration() > FTimespan::Zero())
-			? static_cast<float>(FTimespan::Ratio(MediaPlayer->GetDisplayTime(), MediaPlayer->GetDuration()))
-			: 0.0f;
 	}
 
 	return 0.0f;
@@ -101,14 +130,20 @@ float SMediaPlayerSlider::GetPlaybackPosition() const
 
 EVisibility SMediaPlayerSlider::GetScrubberVisibility() const
 {
-	if (const UMediaPlayer* MediaPlayer = MediaPlayerWeak.Get())
+	bool bIsActive = false;
+	for (const FMediaPlayerEntry& Entry : MediaPlayerEntries)
 	{
-		return (MediaPlayer->SupportsScrubbing() || MediaPlayer->SupportsSeeking())
-															? EVisibility::Visible
-															: VisibilityWhenInactive;
+		if (const UMediaPlayer* MediaPlayer = Entry.MediaPlayerWeak.Get())
+		{
+			bIsActive = (MediaPlayer->SupportsScrubbing() || MediaPlayer->SupportsSeeking());
+			if (!bIsActive)
+			{
+				break; // If any player is inactive, consider widget inactive.
+			}
+		}
 	}
 
-	return VisibilityWhenInactive;
+	return bIsActive ? EVisibility::Visible : VisibilityWhenInactive;
 }
 
 void SMediaPlayerSlider::SetSliderHandleColor(const FSlateColor& InSliderColor)
diff --git a/Engine/Plugins/Media/MediaPlayerEditor/Source/MediaPlayerEditor/Private/Widgets/SMediaPlayerSlider.h b/Engine/Plugins/Media/MediaPlayerEditor/Source/MediaPlayerEditor/Private/Widgets/SMediaPlayerSlider.h
--- a/Engine/Plugins/Media/MediaPlayerEditor/Source/MediaPlayerEditor/Private/Widgets/SMediaPlayerSlider.h
+++ b/Engine/Plugins/Media/MediaPlayerEditor/Source/MediaPlayerEditor/Private/Widgets/SMediaPlayerSlider.h
@@ -28,7 +28,7 @@ public:
 	SLATE_END_ARGS()
 
 	/** Constructs this widget with InArgs */
-	void Construct(const FArguments& InArgs, UMediaPlayer* InMediaPlayer);
+	void Construct(const FArguments& InArgs, const TArrayView<TWeakObjectPtr<UMediaPlayer>> InMediaPlayers);
 
 	//~Begin IMediaPlayerScrubber
 	virtual void SetSliderHandleColor(const FSlateColor& InSliderColor) override;
@@ -47,18 +47,24 @@ private:
 	/** The scrubber visibility when inactive. */
 	EVisibility VisibilityWhenInactive = EVisibility::Hidden;
 
-	/** Pointer to the media player that is being viewed. */
-	TWeakObjectPtr<UMediaPlayer> MediaPlayerWeak;
-
-	/** The playback rate prior to scrubbing. */
-	float PreScrubRate = 0.0f;
+	struct FMediaPlayerEntry
+	{
+		/** Pointer to the media players that is are viewed. */
+		TWeakObjectPtr<UMediaPlayer> MediaPlayerWeak;
+		
+		/** The playback rate prior to scrubbing. */
+		float PreScrubRate = 0.0f;
+		/** The value currently being scrubbed to. */
+		float ScrubValue = 0.0f;
+		
+		/** The last value set with media player while scrubbing. */
+		float LastScrubValue = -1.0f;
 
+		FMediaPlayerEntry(const TWeakObjectPtr<UMediaPlayer>& InMediaPlayerWeak) : MediaPlayerWeak(InMediaPlayerWeak) { }
+	};
+	
+	TArray<FMediaPlayerEntry> MediaPlayerEntries;
+	
 	/** Holds the scrubber slider. */
 	TSharedPtr<SSlider> ScrubberSlider;
-
-	/** The value currently being scrubbed to. */
-	float ScrubValue = 0.0f;
-
-	/** The last value set with media player while scrubbing. */
-	float LastScrubValue = -1.0f;
 };
diff --git a/Engine/Plugins/Media/MediaPlayerEditor/Source/MediaPlayerEditor/Public/MediaPlayerEditorModule.h b/Engine/Plugins/Media/MediaPlayerEditor/Source/MediaPlayerEditor/Public/MediaPlayerEditorModule.h
--- a/Engine/Plugins/Media/MediaPlayerEditor/Source/MediaPlayerEditor/Public/MediaPlayerEditorModule.h
+++ b/Engine/Plugins/Media/MediaPlayerEditor/Source/MediaPlayerEditor/Public/MediaPlayerEditorModule.h
@@ -8,7 +8,6 @@
 #include "Templates/SharedPointer.h"
 #include "Widgets/SCompoundWidget.h"
 
-class FImgMediaPlayer;
 class ISlateStyle;
 class UMediaPlayer;
 
@@ -45,5 +44,14 @@ public:
 	 * @param InStyle: the style chosen for this slider widget
 	 * @return Scrubber Widget
 	 */
+	UE_DEPRECATED(5.5, "Use version with TArrayView instead")
 	virtual TSharedRef<IMediaPlayerSlider> CreateMediaPlayerSliderWidget(UMediaPlayer* InMediaPlayer, const FSliderStyle& InStyle = FCoreStyle::Get().GetWidgetStyle<FSliderStyle>("Slider")) = 0;
+
+	/**
+	 * Creates a Widget to visualize playback time and scrub the content played by Media Players
+	 * @param InMediaPlayers: the players affected by the widget
+	 * @param InStyle: the style chosen for this slider widget
+	 * @return Scrubber Widget
+	 */
+	virtual TSharedRef<IMediaPlayerSlider> CreateMediaPlayerSliderWidget(const TArrayView<TWeakObjectPtr<UMediaPlayer>> InMediaPlayers, const FSliderStyle& InStyle = FCoreStyle::Get().GetWidgetStyle<FSliderStyle>("Slider")) = 0;
 };
