[MediaPlate] Mutli-selection seek slider ignores "failing" players.

Changing the logic to be more permissive, if any player in the selection can scrub or seek, the slider will work with those players, and ignore the ones that can seek or scrub.

Also, the slider position is determined first by players with a valid video track.

#rnx

diff --git a/Engine/Plugins/Media/MediaPlayerEditor/Source/MediaPlayerEditor/Private/Widgets/SMediaPlayerSlider.cpp b/Engine/Plugins/Media/MediaPlayerEditor/Source/MediaPlayerEditor/Private/Widgets/SMediaPlayerSlider.cpp
--- a/Engine/Plugins/Media/MediaPlayerEditor/Source/MediaPlayerEditor/Private/Widgets/SMediaPlayerSlider.cpp
+++ b/Engine/Plugins/Media/MediaPlayerEditor/Source/MediaPlayerEditor/Private/Widgets/SMediaPlayerSlider.cpp
@@ -35,21 +35,18 @@ END_SLATE_FUNCTION_BUILD_OPTIMIZATION
 
 bool SMediaPlayerSlider::DoesMediaPlayerSupportSeeking() const
 {
-	// All players must support seek.
-	bool bSupportSeek = false;
+	// Return true if at least one player supports seek.
 	for (const FMediaPlayerEntry& Entry : MediaPlayerEntries)
 	{
 		if (const UMediaPlayer* MediaPlayer = Entry.MediaPlayerWeak.Get())
 		{
-			bSupportSeek = MediaPlayer->SupportsSeeking();
-			if (!bSupportSeek)
+			if (MediaPlayer->SupportsSeeking())
 			{
-				return false;
+				return true;
 			}
 		}
 	}
-
-	return bSupportSeek;
+	return false;
 }
 
 void SMediaPlayerSlider::OnScrubBegin()
@@ -77,7 +74,7 @@ void SMediaPlayerSlider::OnScrubEnd()
 		if (UMediaPlayer* MediaPlayer = Entry.MediaPlayerWeak.Get())
 		{
 			// Set playback position to scrub value when drag ends
-			if (Entry.LastScrubValue != Entry.ScrubValue)
+			if (Entry.LastScrubValue != Entry.ScrubValue && MediaPlayer->SupportsSeeking())
 			{
 				MediaPlayer->Seek(MediaPlayer->GetDuration() * Entry.ScrubValue);
 			}
@@ -109,19 +106,25 @@ void SMediaPlayerSlider::Seek(float InPlaybackPosition)
 
 float SMediaPlayerSlider::GetPlaybackPosition() const
 {
-	// Note: returns first one. All scrub positions should match.
-	for (const FMediaPlayerEntry& Entry : MediaPlayerEntries)
+	// All scrub positions should match, so search for the first valid player. 
+	// Give priority to players with a video track.
+	const FMediaPlayerEntry* PlayerEntry = FindValidPlayerEntryForTrackType(EMediaPlayerTrack::Video);
+	if (!PlayerEntry)
 	{
-		if (const UMediaPlayer* MediaPlayer = Entry.MediaPlayerWeak.Get())
+		// Fallback to audio track second (not all players support it).
+		PlayerEntry = FindValidPlayerEntryForTrackType(EMediaPlayerTrack::Audio);
+	}
+
+	if (PlayerEntry)
+	{
+		if (ScrubberSlider->HasMouseCapture())
 		{
-			if (ScrubberSlider->HasMouseCapture())
-			{
-				return Entry.ScrubValue;
-			}
+			return PlayerEntry->ScrubValue;
+		}
 
-			return (MediaPlayer->GetDuration() > FTimespan::Zero())
-				? static_cast<float>(FTimespan::Ratio(MediaPlayer->GetDisplayTime(), MediaPlayer->GetDuration()))
-				: 0.0f;
+		if (const UMediaPlayer* MediaPlayer = PlayerEntry->MediaPlayerWeak.Get())
+		{
+			return static_cast<float>(FTimespan::Ratio(MediaPlayer->GetDisplayTime(), MediaPlayer->GetDuration()));
 		}
 	}
 
@@ -136,9 +139,9 @@ EVisibility SMediaPlayerSlider::GetScrubberVisibility() const
 		if (const UMediaPlayer* MediaPlayer = Entry.MediaPlayerWeak.Get())
 		{
 			bIsActive = (MediaPlayer->SupportsScrubbing() || MediaPlayer->SupportsSeeking());
-			if (!bIsActive)
+			if (bIsActive)
 			{
-				break; // If any player is inactive, consider widget inactive.
+				break; // If any player is active, consider widget active.
 			}
 		}
 	}
@@ -166,3 +169,19 @@ void SMediaPlayerSlider::SetVisibleWhenInactive(EVisibility InVisibility)
 {
 	VisibilityWhenInactive = InVisibility;
 }
+
+/** Find a valid player entry for the given track type. */
+const SMediaPlayerSlider::FMediaPlayerEntry* SMediaPlayerSlider::FindValidPlayerEntryForTrackType(EMediaPlayerTrack InTrackType) const
+{
+	for (const FMediaPlayerEntry& Entry : MediaPlayerEntries)
+	{
+		if (const UMediaPlayer* MediaPlayer = Entry.MediaPlayerWeak.Get())
+		{
+			if (MediaPlayer->GetNumTracks(InTrackType) > 0 && MediaPlayer->GetDuration() > FTimespan::Zero())
+			{
+				return &Entry;
+			}
+		}
+	}
+	return nullptr;
+}
\ No newline at end of file
diff --git a/Engine/Plugins/Media/MediaPlayerEditor/Source/MediaPlayerEditor/Private/Widgets/SMediaPlayerSlider.h b/Engine/Plugins/Media/MediaPlayerEditor/Source/MediaPlayerEditor/Private/Widgets/SMediaPlayerSlider.h
--- a/Engine/Plugins/Media/MediaPlayerEditor/Source/MediaPlayerEditor/Private/Widgets/SMediaPlayerSlider.h
+++ b/Engine/Plugins/Media/MediaPlayerEditor/Source/MediaPlayerEditor/Private/Widgets/SMediaPlayerSlider.h
@@ -10,6 +10,7 @@
 
 class SSlider;
 class UMediaPlayer;
+enum class EMediaPlayerTrack : uint8;
 
 /**
  * Implements a scrubber to visualize the current playback position of a Media Player
@@ -62,6 +63,8 @@ private:
 
 		FMediaPlayerEntry(const TWeakObjectPtr<UMediaPlayer>& InMediaPlayerWeak) : MediaPlayerWeak(InMediaPlayerWeak) { }
 	};
+
+	const FMediaPlayerEntry* FindValidPlayerEntryForTrackType(EMediaPlayerTrack InTrackType) const;
 	
 	TArray<FMediaPlayerEntry> MediaPlayerEntries;
 	
