Media I/O - Can not add playlist in Media Plate window

- Add entry button is now hidden by default, and if the current playlist is invalid.
- Added text to let the user know if a playlist source is selected and/or if it's empty

In case it becomes desireable to have the Add entry button in the future the showing/hiding behavior could be further implemented by letting the Widget know about the Media Plate Resource and not just the Media Playlist.
In case the current source type is not a playlist, the button should be hidden.

diff --git a/Engine/Plugins/Media/MediaPlayerEditor/Source/MediaPlayerEditor/Private/Widgets/SMediaPlaylistEditorTracks.cpp b/Engine/Plugins/Media/MediaPlayerEditor/Source/MediaPlayerEditor/Private/Widgets/SMediaPlaylistEditorTracks.cpp
--- a/Engine/Plugins/Media/MediaPlayerEditor/Source/MediaPlayerEditor/Private/Widgets/SMediaPlaylistEditorTracks.cpp
+++ b/Engine/Plugins/Media/MediaPlayerEditor/Source/MediaPlayerEditor/Private/Widgets/SMediaPlaylistEditorTracks.cpp
@@ -11,6 +11,7 @@
 #include "Widgets/Images/SImage.h"
 #include "Widgets/Input/SButton.h"
 #include "Widgets/Layout/SScrollBox.h"
+#include "Widgets/Text/STextBlock.h"
 #include "Styling/AppStyle.h"
 #include "ScopedTransaction.h"
 #include "UObject/UObjectGlobals.h"
@@ -20,8 +21,10 @@
 void SMediaPlaylistEditorTracks::Construct(const FArguments& InArgs, UMediaPlaylist* InMediaPlaylist, const TSharedRef<ISlateStyle>& InStyle)
 {
 	MediaPlaylistPtr = InMediaPlaylist;
+	bShowAddEntryButton = InArgs._ShowAddEntryButton;
 
 	ChildSlot
+	.Padding(4.0f)
 		[
 			SNew(SScrollBox)
 
@@ -40,24 +43,37 @@ void SMediaPlaylistEditorTracks::Construct(const FArguments& InArgs, UMediaPlayl
 							// Add button to add to the playlist.
 							SNew(SButton)
 							.ToolTipText(LOCTEXT("Add_ToolTip", "Add an entry to the playlist."))
-								.VAlign(VAlign_Center)
-								.OnClicked_Lambda([this]() -> FReply
-								{
-									AddToPlaylist();
-									return FReply::Handled();
-								})
-								[
-									SNew(SImage)
-										.ColorAndOpacity(FSlateColor::UseForeground())
-										.Image(FAppStyle::GetBrush("Icons.PlusCircle"))
-								]
+							.VAlign(VAlign_Center)
+							.Visibility(this, &SMediaPlaylistEditorTracks::GetAddEntryButtonVisibility)
+							.OnClicked_Lambda([this]() -> FReply
+							{
+								AddToPlaylist();
+								return FReply::Handled();
+							})
+							[
+								SNew(SImage)
+									.ColorAndOpacity(FSlateColor::UseForeground())
+									.Image(FAppStyle::GetBrush("Icons.PlusCircle"))
+							]
 						]
 				]
 
 			// Container for our media sources.
 			+ SScrollBox::Slot()
 				[
-					SAssignNew(SourcesContainer, SVerticalBox)
+					SNew(SVerticalBox)
+					+ SVerticalBox::Slot()
+					.AutoHeight()
+					[
+						SNew(STextBlock)
+						.Visibility(this, &SMediaPlaylistEditorTracks::GetSourcesContainerHeaderTextVisibility)
+						.Text(this, &SMediaPlaylistEditorTracks::GetSourcesContainerHeaderText)
+					]
+					+ SVerticalBox::Slot()
+					.AutoHeight()
+					[
+						SAssignNew(SourcesContainer, SVerticalBox)
+					]
 				]
 		];
 
@@ -206,4 +222,41 @@ void SMediaPlaylistEditorTracks::OnObjectTransacted(UObject* Object, const FTran
 	}
 }
 
+EVisibility SMediaPlaylistEditorTracks::GetAddEntryButtonVisibility() const
+{
+	return bShowAddEntryButton && MediaPlaylistPtr.IsValid() ? EVisibility::Visible : EVisibility::Collapsed;
+}
+
+EVisibility SMediaPlaylistEditorTracks::GetSourcesContainerHeaderTextVisibility() const
+{
+	TStrongObjectPtr<UMediaPlaylist> MediaPlaylistPinned = MediaPlaylistPtr.Pin();
+	if (!MediaPlaylistPinned)
+	{
+		return EVisibility::Visible;
+	}
+
+	if (MediaPlaylistPinned->Num() == 0)
+	{
+		return EVisibility::Visible;
+	}
+
+	return EVisibility::Collapsed;
+}
+
+FText SMediaPlaylistEditorTracks::GetSourcesContainerHeaderText() const
+{
+	TStrongObjectPtr<UMediaPlaylist> MediaPlaylistPinned = MediaPlaylistPtr.Pin();
+	if (!MediaPlaylistPinned)
+	{
+		return LOCTEXT("MissingPlaylistText", "No playlist source selected.");
+	}
+
+	if (MediaPlaylistPinned->Num() == 0)
+	{
+		return LOCTEXT("EmptyPlaylistText", "Selected playlist source is empty.");
+	}
+
+	return FText();
+}
+
 #undef LOCTEXT_NAMESPACE
diff --git a/Engine/Plugins/Media/MediaPlayerEditor/Source/MediaPlayerEditor/Public/Widgets/SMediaPlaylistEditorTracks.h b/Engine/Plugins/Media/MediaPlayerEditor/Source/MediaPlayerEditor/Public/Widgets/SMediaPlaylistEditorTracks.h
--- a/Engine/Plugins/Media/MediaPlayerEditor/Source/MediaPlayerEditor/Public/Widgets/SMediaPlaylistEditorTracks.h
+++ b/Engine/Plugins/Media/MediaPlayerEditor/Source/MediaPlayerEditor/Public/Widgets/SMediaPlaylistEditorTracks.h
@@ -17,8 +17,11 @@ class MEDIAPLAYEREDITOR_API SMediaPlaylistEditorTracks
 	: public SCompoundWidget
 {
 public:
+	SLATE_BEGIN_ARGS(SMediaPlaylistEditorTracks)
+		: _ShowAddEntryButton(false)
+		{}
 
-	SLATE_BEGIN_ARGS(SMediaPlaylistEditorTracks) { }
+		SLATE_ARGUMENT(bool, ShowAddEntryButton)
 	SLATE_END_ARGS()
 
 public:
@@ -40,6 +43,12 @@ protected:
 	/** Pointer to the container for the media sources. */
 	TSharedPtr<SVerticalBox> SourcesContainer;
 
+	/**
+	 * Flag to enable/disable a button which can be used to add items to the currently selected Playlist.
+	 * Set by "ShowAddEntryButton" slate argument. If no Playlist is selected, button will be hidden anyway.
+	 */
+	bool bShowAddEntryButton = false;
+
 	/**
 	 * This should be called when the playlist changes so we can update the UI.
 	 */
@@ -67,4 +76,13 @@ protected:
 
 	/** Engine callback to detect when the media is changed. */
 	void OnObjectTransacted(UObject* Object, const FTransactionObjectEvent& TransactionObjectEvent);
+
+	/** Returns the current visibility for the button used to add items to the currently selected Playlist, if any */
+	EVisibility GetAddEntryButtonVisibility() const;
+
+	/** Returns the current visibility of the header area (before or in place of the current Playlist items) */
+	EVisibility GetSourcesContainerHeaderTextVisibility() const;
+
+	/** Returns the text to be displayed in the header area (before or in place of the current Playlist items) */
+	FText GetSourcesContainerHeaderText() const;
 };
