MediaPlayerEditor: Extending the per-platform media player selection to include all platforms for which ini files are found instead of vanilla platforms, making it easier to adjust this setting for all platforms on just one desktop machine instead of having to change this on a Windows, Mac and Linux machine.

#rnx

diff --git a/Engine/Plugins/Media/MediaPlayerEditor/Source/MediaPlayerEditor/Private/Customizations/BaseMediaSourceCustomization.cpp b/Engine/Plugins/Media/MediaPlayerEditor/Source/MediaPlayerEditor/Private/Customizations/BaseMediaSourceCustomization.cpp
--- a/Engine/Plugins/Media/MediaPlayerEditor/Source/MediaPlayerEditor/Private/Customizations/BaseMediaSourceCustomization.cpp
+++ b/Engine/Plugins/Media/MediaPlayerEditor/Source/MediaPlayerEditor/Private/Customizations/BaseMediaSourceCustomization.cpp
@@ -133,27 +133,14 @@ TSharedRef<SWidget> FBaseMediaSourceCustomization::MakePlatformPlayerNamesValueW
 	}
 
 	// get available platforms
-	TArray<const PlatformInfo::FTargetPlatformInfo*> AvailablePlatforms;
-
-	for (const PlatformInfo::FTargetPlatformInfo* PlatformInfo : PlatformInfo::GetVanillaPlatformInfoArray())
-	{
-		if (PlatformInfo->PlatformType == EBuildTargetType::Game)
-		{
-			AvailablePlatforms.Add(PlatformInfo);
-		}
-	}
-
-	Algo::Sort(AvailablePlatforms, [](const PlatformInfo::FTargetPlatformInfo* One, const PlatformInfo::FTargetPlatformInfo* Two) -> bool
-	{
-		return One->DisplayName.CompareTo(Two->DisplayName) < 0;
-	});
+	const TArray<const FDataDrivenPlatformInfo*>& PlatformInfos = FDataDrivenPlatformInfoRegistry::GetSortedPlatformInfos(EPlatformInfoType::TruePlatformsOnly);
 
 	// build value widget
 	TSharedRef<SGridPanel> PlatformPanel = SNew(SGridPanel);
 
-	for (int32 PlatformIndex = 0; PlatformIndex < AvailablePlatforms.Num(); ++PlatformIndex)
+	for (int32 PlatformIndex = 0; PlatformIndex < PlatformInfos.Num(); ++PlatformIndex)
 	{
-		const PlatformInfo::FTargetPlatformInfo* Platform = AvailablePlatforms[PlatformIndex];
+		const FDataDrivenPlatformInfo* const Platform = PlatformInfos[PlatformIndex];
 
 		// platform icon
 		PlatformPanel->AddSlot(0, PlatformIndex)
@@ -169,7 +156,7 @@ TSharedRef<SWidget> FBaseMediaSourceCustomization::MakePlatformPlayerNamesValueW
 			.VAlign(VAlign_Center)
 			[
 				SNew(STextBlock)
-					.Text(Platform->DisplayName)
+					.Text(FText::FromName(Platform->IniPlatformName))
 			];
 
 		// player combo box
diff --git a/Engine/Plugins/Media/MediaPlayerEditor/Source/MediaPlayerEditor/Private/Customizations/PlatformMediaSourceCustomization.cpp b/Engine/Plugins/Media/MediaPlayerEditor/Source/MediaPlayerEditor/Private/Customizations/PlatformMediaSourceCustomization.cpp
--- a/Engine/Plugins/Media/MediaPlayerEditor/Source/MediaPlayerEditor/Private/Customizations/PlatformMediaSourceCustomization.cpp
+++ b/Engine/Plugins/Media/MediaPlayerEditor/Source/MediaPlayerEditor/Private/Customizations/PlatformMediaSourceCustomization.cpp
@@ -72,33 +72,14 @@ TSharedRef<SWidget> FPlatformMediaSourceCustomization::MakePlatformMediaSourcesV
 	const TArray<IMediaPlayerFactory*>& PlayerFactories = MediaModule->GetPlayerFactories();
 
 	// get available platforms
-	TArray<PlatformInfo::FTargetPlatformInfo*> AvailablePlatforms;
-
-	for (PlatformInfo::FTargetPlatformInfo* PlatformInfo : PlatformInfo::GetVanillaPlatformInfoArray())
-	{
-		if (PlatformInfo->PlatformType == EBuildTargetType::Game)
-		{
-			if (PlatformInfo->Name == TEXT("TVOS"))
-			{
-				continue; // tvOS is just iOS for now
-			}
-
-			AvailablePlatforms.Add(PlatformInfo);
-		}
-	}
-
-	// sort available platforms alphabetically
-	Algo::Sort(AvailablePlatforms, [](const PlatformInfo::FTargetPlatformInfo* One, const PlatformInfo::FTargetPlatformInfo* Two) -> bool
-	{
-		return One->DisplayName.CompareTo(Two->DisplayName) < 0;
-	});
+	const TArray<const FDataDrivenPlatformInfo*>& PlatformInfos = FDataDrivenPlatformInfoRegistry::GetSortedPlatformInfos(EPlatformInfoType::TruePlatformsOnly);
 
 	// build value widget
 	TSharedRef<SGridPanel> PlatformPanel = SNew(SGridPanel);
 
-	for (int32 PlatformIndex = 0; PlatformIndex < AvailablePlatforms.Num(); ++PlatformIndex)
+	for (int32 PlatformIndex = 0; PlatformIndex < PlatformInfos.Num(); ++PlatformIndex)
 	{
-		const PlatformInfo::FTargetPlatformInfo* Platform = AvailablePlatforms[PlatformIndex];
+		const FDataDrivenPlatformInfo* const Platform = PlatformInfos[PlatformIndex];
 
 		// platform icon
 		PlatformPanel->AddSlot(0, PlatformIndex)
@@ -114,7 +95,7 @@ TSharedRef<SWidget> FPlatformMediaSourceCustomization::MakePlatformMediaSourcesV
 			.VAlign(VAlign_Center)
 			[
 				SNew(STextBlock)
-					.Text(Platform->DisplayName)
+					.Text(FText::FromName(Platform->IniPlatformName))
 			];
 
 		// player combo box
