Media: Added query to retrieve keyframe interval to determine seek performace for display in the media info panel

#rnx

diff --git a/Engine/Plugins/Media/ElectraPlayer/Source/ElectraProtron/Private/Player/ElectraProtronPlayerImpl.cpp b/Engine/Plugins/Media/ElectraPlayer/Source/ElectraProtron/Private/Player/ElectraProtronPlayerImpl.cpp
--- a/Engine/Plugins/Media/ElectraPlayer/Source/ElectraProtron/Private/Player/ElectraProtronPlayerImpl.cpp
+++ b/Engine/Plugins/Media/ElectraPlayer/Source/ElectraProtron/Private/Player/ElectraProtronPlayerImpl.cpp
@@ -16,6 +16,7 @@ namespace ElectraProtronOptionNames
 {
 	const FName StartTimecodeValue(TEXT("StartTimecodeValue"));				// maybe use:  UMediaPlayer::MediaInfoNameSourceNumTiles.Resolve()
 	const FName StartTimecodeFrameRate(TEXT("StartTimecodeFrameRate"));
+	const FName KeyframeInterval(TEXT("KeyframeInterval"));
 }
 
 
@@ -910,23 +911,25 @@ void FElectraProtronPlayer::FImpl::GetTrackCodecInfo(ElectraProtronUtils::FCodec
 
 FVariant FElectraProtronPlayer::FImpl::GetMediaInfo(FName InInfoName)
 {
-	if (InInfoName == ElectraProtronOptionNames::StartTimecodeValue || InInfoName == ElectraProtronOptionNames::StartTimecodeFrameRate)
+	if (InInfoName == ElectraProtronOptionNames::StartTimecodeValue || InInfoName == ElectraProtronOptionNames::StartTimecodeFrameRate || InInfoName == ElectraProtronOptionNames::KeyframeInterval)
 	{
 		auto ci = CodecTypeIndex(ElectraProtronUtils::FCodecInfo::EType::Video);
 		int32 SelectedVideoTrackIndex = TrackSelection.SelectedTrackIndex[ci];
 		if (SelectedVideoTrackIndex >= 0 && SelectedVideoTrackIndex < UsableTrackArrayIndicesByType[ci].Num())
 		{
 			const FTrackInfo& ti(*Tracks[UsableTrackArrayIndicesByType[ci][SelectedVideoTrackIndex]]);
-			if (ti.FirstSampleTimecode.IsSet())
+
+			if (InInfoName == ElectraProtronOptionNames::StartTimecodeValue && ti.FirstSampleTimecode.IsSet())
 			{
-				if (InInfoName == ElectraProtronOptionNames::StartTimecodeValue)
-				{
-					return FVariant(ti.FirstSampleTimecode.GetValue().Timecode);
-				}
-				else //if (InInfoName == ElectraProtronOptionNames::StartTimecodeFrameRate)
-				{
-					return FVariant(ti.FirstSampleTimecode.GetValue().Framerate);
-				}
+				return FVariant(ti.FirstSampleTimecode.GetValue().Timecode);
+			}
+			else if (InInfoName == ElectraProtronOptionNames::StartTimecodeFrameRate && ti.FirstSampleTimecode.IsSet())
+			{
+				return FVariant(ti.FirstSampleTimecode.GetValue().Framerate);
+			}
+			else if (InInfoName == ElectraProtronOptionNames::KeyframeInterval)
+			{
+				return FVariant(ti.bIsKeyframeOnlyFormat ? (int32)1 : (int32)0);
 			}
 		}
 	}
diff --git a/Engine/Plugins/Media/MediaPlayerEditor/Source/MediaPlayerEditor/Private/Widgets/SMediaPlayerEditorMediaDetails.cpp b/Engine/Plugins/Media/MediaPlayerEditor/Source/MediaPlayerEditor/Private/Widgets/SMediaPlayerEditorMediaDetails.cpp
--- a/Engine/Plugins/Media/MediaPlayerEditor/Source/MediaPlayerEditor/Private/Widgets/SMediaPlayerEditorMediaDetails.cpp
+++ b/Engine/Plugins/Media/MediaPlayerEditor/Source/MediaPlayerEditor/Private/Widgets/SMediaPlayerEditorMediaDetails.cpp
@@ -119,6 +119,14 @@ void SMediaPlayerEditorMediaDetails::Construct(const FArguments& InArgs,
 								[
 									SAssignNew(StartTimecodeText, STextBlock)
 								]
+							// Seek Performance
+							+ SVerticalBox::Slot()
+								.AutoHeight()
+								.VAlign(VAlign_Center)
+								.Padding(4.0f)
+								[
+									SAssignNew(SeekPerformance, STextBlock)
+								]
 						]
 				]
 
@@ -148,6 +156,7 @@ void SMediaPlayerEditorMediaDetails::UpdateDetails()
 	int32 SurfaceWidth = 0;
 	int32 SurfaceHeight = 0;
 	FString StartTimecode;
+	int32 KeyframeInterval = -1;
 
 	// Get player info.
 	if (MediaPlayer != nullptr)
@@ -158,6 +167,7 @@ void SMediaPlayerEditorMediaDetails::UpdateDetails()
 		MediaPlayer->GetMediaInfo<FIntPoint>(NumTiles, UMediaPlayer::MediaInfoNameSourceNumTiles.Resolve());
 		NumTotalTiles = NumTiles.X * NumTiles.Y;
 		MediaPlayer->GetMediaInfo<FString>(StartTimecode, UMediaPlayer::MediaInfoNameStartTimecodeValue.Resolve());
+		MediaPlayer->GetMediaInfo<int32>(KeyframeInterval, UMediaPlayer::MediaInfoNameKeyframeInterval.Resolve());
 	}
 
 	// Get texture info,
@@ -177,7 +187,7 @@ void SMediaPlayerEditorMediaDetails::UpdateDetails()
 	// Update text.
 	FormatText->SetText(FText::Format(LOCTEXT("Format", "Format: {0}"),
 		FText::FromString(Format)));
-	FrameRateText->SetText(FText::Format(LOCTEXT("FrameRate", "Frame Rate: {0}"), 
+	FrameRateText->SetText(FText::Format(LOCTEXT("FrameRate", "Frame Rate: {0}"),
 		FText::AsNumber(FrameRate)));
 	LODBiasText->SetText(FText::Format(LOCTEXT("LODBias", "Combined LOD Bias: {0}"),
 		FText::AsNumber(LODBias)));
@@ -190,12 +200,29 @@ void SMediaPlayerEditorMediaDetails::UpdateDetails()
 		FText::AsNumber(SurfaceWidth), FText::AsNumber(SurfaceHeight)));
 	ResourceSizeText->SetText(FText::Format(LOCTEXT("ResourceSize", "Resource Size: {0} KB"),
 		FText::AsNumber(ResourceSize)));
-	
+
 	if (!StartTimecode.IsEmpty())
 	{
 		StartTimecodeText->SetText(FText::Format(LOCTEXT("StartTimecode", "Start Timecode: {0}"), FText::FromString(StartTimecode)));
 	}
 	StartTimecodeText->SetVisibility(!StartTimecode.IsEmpty() ? EVisibility::Visible : EVisibility::Collapsed);
+
+	/*
+		The keyframe interval is one of:
+		  -1 : no information returned (unknown)
+		   0 : unknown keyframe spacing, not every frame is a keyframe but the spacing is variable or cannot be determined
+		   1 : every frame is a keyframe
+		  >1 : every n'th frame is a keyframe
+	*/
+	if (KeyframeInterval >= 0)
+	{
+		SeekPerformance->SetText(FText::Format(LOCTEXT("SeekPerformance", "Seek Performance: {0}"), KeyframeInterval == 1 ? LOCTEXT("SeekPerformanceF", "Fast") : LOCTEXT("SeekPerformanceS", "Slow (GOP codec)")));
+		SeekPerformance->SetVisibility(EVisibility::Visible);
+	}
+	else
+	{
+		SeekPerformance->SetVisibility(EVisibility::Collapsed);
+	}
 }
 
 #undef LOCTEXT_NAMESPACE
diff --git a/Engine/Plugins/Media/MediaPlayerEditor/Source/MediaPlayerEditor/Public/Widgets/SMediaPlayerEditorMediaDetails.h b/Engine/Plugins/Media/MediaPlayerEditor/Source/MediaPlayerEditor/Public/Widgets/SMediaPlayerEditorMediaDetails.h
--- a/Engine/Plugins/Media/MediaPlayerEditor/Source/MediaPlayerEditor/Public/Widgets/SMediaPlayerEditorMediaDetails.h
+++ b/Engine/Plugins/Media/MediaPlayerEditor/Source/MediaPlayerEditor/Public/Widgets/SMediaPlayerEditorMediaDetails.h
@@ -29,7 +29,7 @@ public:
 	 * @param InArgs			The declaration data for this widget.
 	 * @param InMediaPlayer		The MediaPlayer to show the details for.
 	 * @param InMediaTexture	The MediaTexture to show the details for.
-	 * 
+	 *
 	 */
 	void Construct(const FArguments& InArgs, UMediaPlayer* InMediaPlate, UMediaTexture* InMediaTexture);
 
@@ -58,4 +58,5 @@ private:
 	TSharedPtr<STextBlock> ResolutionText;
 	TSharedPtr<STextBlock> ResourceSizeText;
 	TSharedPtr<STextBlock> StartTimecodeText;
+	TSharedPtr<STextBlock> SeekPerformance;
 };
diff --git a/Engine/Source/Runtime/MediaAssets/Private/Assets/MediaPlayer.cpp b/Engine/Source/Runtime/MediaAssets/Private/Assets/MediaPlayer.cpp
--- a/Engine/Source/Runtime/MediaAssets/Private/Assets/MediaPlayer.cpp
+++ b/Engine/Source/Runtime/MediaAssets/Private/Assets/MediaPlayer.cpp
@@ -37,6 +37,8 @@ FLazyName UMediaPlayer::MediaInfoNameSourceNumMips(TEXT("SourceNumMips"));
 FLazyName UMediaPlayer::MediaInfoNameSourceNumTiles(TEXT("SourceNumTiles"));
 FLazyName UMediaPlayer::MediaInfoNameStartTimecodeValue(TEXT("StartTimecodeValue"));
 FLazyName UMediaPlayer::MediaInfoNameStartTimecodeFrameRate(TEXT("StartTimecodeFrameRate"));
+FLazyName UMediaPlayer::MediaInfoNameKeyframeInterval(TEXT("KeyframeInterval"));
+
 
 /* UMediaPlayer structors
  *****************************************************************************/
diff --git a/Engine/Source/Runtime/MediaAssets/Public/MediaPlayer.h b/Engine/Source/Runtime/MediaAssets/Public/MediaPlayer.h
--- a/Engine/Source/Runtime/MediaAssets/Public/MediaPlayer.h
+++ b/Engine/Source/Runtime/MediaAssets/Public/MediaPlayer.h
@@ -1310,6 +1310,8 @@ public:
 	static MEDIAASSETS_API FLazyName MediaInfoNameStartTimecodeValue;
 	/** Start Timecode frame rate. String. */
 	static MEDIAASSETS_API FLazyName MediaInfoNameStartTimecodeFrameRate;
+	/** Keyframe interval. Int32. */
+	static MEDIAASSETS_API FLazyName MediaInfoNameKeyframeInterval;
 
 	/**
 	 * Duration of samples to cache ahead of the play head.
