[MediaPlayerEditor] Updating comment to match class hierarchy.

diff --git a/Engine/Plugins/Media/MediaPlayerEditor/Source/MediaPlayerEditor/Private/Widgets/SMediaPlayerSlider.h b/Engine/Plugins/Media/MediaPlayerEditor/Source/MediaPlayerEditor/Private/Widgets/SMediaPlayerSlider.h
--- a/Engine/Plugins/Media/MediaPlayerEditor/Source/MediaPlayerEditor/Private/Widgets/SMediaPlayerSlider.h
+++ b/Engine/Plugins/Media/MediaPlayerEditor/Source/MediaPlayerEditor/Private/Widgets/SMediaPlayerSlider.h
@@ -31,11 +31,11 @@ public:
 	/** Constructs this widget with InArgs */
 	void Construct(const FArguments& InArgs, const TArrayView<TWeakObjectPtr<UMediaPlayer>> InMediaPlayers);
 
-	//~Begin IMediaPlayerScrubber
+	//~Begin IMediaPlayerSlider
 	virtual void SetSliderHandleColor(const FSlateColor& InSliderColor) override;
 	virtual void SetSliderBarColor(const FSlateColor& InSliderColor) override;
 	virtual void SetVisibleWhenInactive(EVisibility InVisibility) override;
-	//~End IMediaPlayerScrubber
+	//~End IMediaPlayerSlider
 
 private:
 	bool DoesMediaPlayerSupportSeeking() const;
