[MediaPlayerEditor] Validate the globally specified desired player before opening the media source to avoid mistakes with incompatible media sources.

If the media source can't be opened by the desired player, it will be reset so the player from the media source can be selected.

For this to work, CanPlayUrl must give us a usable answer. In the case of Protron, it is not.

In terms of score for the "file" scheme on all "windows/mac" platform
Electra: 100 or 20
Wmf: 80
Protron: 1

So protron will always loose unless it really is the only one available.

diff --git a/Engine/Plugins/Media/ElectraPlayer/Source/ElectraProtronFactory/Private/ElectraProtronFactoryModule.cpp b/Engine/Plugins/Media/ElectraPlayer/Source/ElectraProtronFactory/Private/ElectraProtronFactoryModule.cpp
--- a/Engine/Plugins/Media/ElectraPlayer/Source/ElectraProtronFactory/Private/ElectraProtronFactoryModule.cpp
+++ b/Engine/Plugins/Media/ElectraPlayer/Source/ElectraProtronFactory/Private/ElectraProtronFactoryModule.cpp
@@ -75,8 +75,8 @@ public:
 		// so we cannot check for this.
 		if (SupportedFileExtensions.Contains(LowerCaseExtension))
 		{
-	// For the time being Protron *MUST* be selected explicitly.
-			return 0;
+			// Low score to select other players instead.
+			return 1;
 		}
 		return 0;
 	}
diff --git a/Engine/Plugins/Media/MediaPlayerEditor/Source/MediaPlayerEditor/Private/Toolkits/MediaSourceEditorToolkit.cpp b/Engine/Plugins/Media/MediaPlayerEditor/Source/MediaPlayerEditor/Private/Toolkits/MediaSourceEditorToolkit.cpp
--- a/Engine/Plugins/Media/MediaPlayerEditor/Source/MediaPlayerEditor/Private/Toolkits/MediaSourceEditorToolkit.cpp
+++ b/Engine/Plugins/Media/MediaPlayerEditor/Source/MediaPlayerEditor/Private/Toolkits/MediaSourceEditorToolkit.cpp
@@ -6,17 +6,23 @@
 #include "EditorReimportHandler.h"
 #include "Engine/TextureRenderTarget2D.h"
 #include "Framework/MultiBox/MultiBoxBuilder.h"
+#include "IMediaModule.h"
+#include "IMediaPlayerFactory.h"
 #include "MediaPlayer.h"
 #include "MediaPlaylist.h"
 #include "MediaSource.h"
 #include "MediaTexture.h"
 #include "Models/MediaPlayerEditorCommands.h"
+#include "Modules/ModuleManager.h"
+#include "PixelShaderUtils.h"
+#include "PostProcess/DrawRectangle.h"
 #include "RenderGraphBuilder.h"
 #include "RenderGraphUtils.h"
 #include "ScreenPass.h"
 #include "SlateOptMacros.h"
 #include "Styling/AppStyle.h"
 #include "TextureResource.h"
+#include "UObject/Package.h"
 #include "Widgets/Docking/SDockTab.h"
 #include "Widgets/SMediaPlayerEditorDetails.h"
 #include "Widgets/SMediaPlayerEditorMediaDetails.h"
@@ -24,9 +30,6 @@
 #include "Widgets/SMediaPlayerEditorStats.h"
 #include "Widgets/SMediaPlayerEditorViewer.h"
 #include "Widgets/SMediaSourceEditorDetails.h"
-#include "UObject/Package.h"
-#include "PostProcess/DrawRectangle.h"
-#include "PixelShaderUtils.h"
 
 #define LOCTEXT_NAMESPACE "FMediaSourceEditorToolkit"
 
@@ -260,6 +263,26 @@ void FMediaSourceEditorToolkit::PostRedo(bool bSuccess)
 /* FMediaSourceEditorToolkit implementation
  *****************************************************************************/
 
+void FMediaSourceEditorToolkit::ValidateDesiredPlayer(UMediaPlayer* InMediaPlayer, UMediaSource* InMediaSource)
+{
+	// Validate desired player
+	FName DesiredPlayerName = InMediaPlayer->GetDesiredPlayerName();
+	if (DesiredPlayerName != NAME_None)
+	{
+		if (IMediaModule* MediaModule = FModuleManager::LoadModulePtr<IMediaModule>("Media"))
+		{
+			if (IMediaPlayerFactory* PlayerFactory = MediaModule->GetPlayerFactory(DesiredPlayerName))
+			{
+				// Check that the given factory can open the media source, if not, don't use it.
+				if (!PlayerFactory->CanPlayUrl(InMediaSource->GetUrl(), InMediaSource))
+				{
+					InMediaPlayer->SetDesiredPlayerName(NAME_None);
+				}
+			}
+		}
+	}
+}
+
 void FMediaSourceEditorToolkit::BindCommands()
 {
 	const FMediaPlayerEditorCommands& Commands = FMediaPlayerEditorCommands::Get();
@@ -295,6 +318,7 @@ void FMediaSourceEditorToolkit::BindCommands()
 			Options.SetAllAsOptional();
 			Options.InternalCustomOptions.Emplace(MediaPlayerOptionValues::Environment(), MediaPlayerOptionValues::Environment_Preview());
 			Options.InternalCustomOptions.Emplace(MediaPlayerOptionValues::ParseTimecodeInfo(), FVariant());
+			ValidateDesiredPlayer(MediaPlayer, MediaSource);
 			MediaPlayer->OpenSourceWithOptions(MediaSource, Options);
 		}),
 		FCanExecuteAction::CreateLambda([this] { return true; })
diff --git a/Engine/Plugins/Media/MediaPlayerEditor/Source/MediaPlayerEditor/Private/Toolkits/MediaSourceEditorToolkit.h b/Engine/Plugins/Media/MediaPlayerEditor/Source/MediaPlayerEditor/Private/Toolkits/MediaSourceEditorToolkit.h
--- a/Engine/Plugins/Media/MediaPlayerEditor/Source/MediaPlayerEditor/Private/Toolkits/MediaSourceEditorToolkit.h
+++ b/Engine/Plugins/Media/MediaPlayerEditor/Source/MediaPlayerEditor/Private/Toolkits/MediaSourceEditorToolkit.h
@@ -79,6 +79,11 @@ protected:
 	virtual void PostRedo(bool bSuccess) override;
 
 protected:
+	/**
+	 * Validates that the given player can open the media source with the selected desired player.
+	 * If not, the desired player is reset to fallback to the media source selection.
+	 */
+	static void ValidateDesiredPlayer(UMediaPlayer* InMediaPlayer, UMediaSource* InMediaSource);
 
 	/** Binds the UI commands to delegates. */
 	void BindCommands();
