Media: Added displaying the video and audio decoder implementation (if available) being used in the media player details.

#rnx

#Synced-CL Unknown

diff --git a/Engine/Plugins/Media/MediaPlayerEditor/Source/MediaPlayerEditor/Private/Widgets/SMediaPlayerEditorMediaDetails.cpp b/Engine/Plugins/Media/MediaPlayerEditor/Source/MediaPlayerEditor/Private/Widgets/SMediaPlayerEditorMediaDetails.cpp
--- a/Engine/Plugins/Media/MediaPlayerEditor/Source/MediaPlayerEditor/Private/Widgets/SMediaPlayerEditorMediaDetails.cpp
+++ b/Engine/Plugins/Media/MediaPlayerEditor/Source/MediaPlayerEditor/Private/Widgets/SMediaPlayerEditorMediaDetails.cpp
@@ -85,41 +85,42 @@ void SMediaPlayerEditorMediaDetails::Construct(const FArguments& InArgs,
 						[
 							SNew(SVerticalBox)
 
-							// Format.
+							// LOD bias.
 							+ SVerticalBox::Slot()
 								.AutoHeight()
 								.VAlign(VAlign_Center)
 								.Padding(4.0f)
 								[
-									SAssignNew(FormatText, STextBlock)
+									SAssignNew(LODBiasText, STextBlock)
 								]
 
-							// LOD bias.
+							// Num mips.
 							+ SVerticalBox::Slot()
 								.AutoHeight()
 								.VAlign(VAlign_Center)
 								.Padding(4.0f)
 								[
-									SAssignNew(LODBiasText, STextBlock)
+									SAssignNew(NumMipsText, STextBlock)
 								]
 
-							// Num mips.
+							// Num tiles.
 							+ SVerticalBox::Slot()
 								.AutoHeight()
 								.VAlign(VAlign_Center)
 								.Padding(4.0f)
 								[
-									SAssignNew(NumMipsText, STextBlock)
+									SAssignNew(NumTilesText, STextBlock)
 								]
 
-							// Num tiles.
+							// Format.
 							+ SVerticalBox::Slot()
 								.AutoHeight()
 								.VAlign(VAlign_Center)
 								.Padding(4.0f)
 								[
-									SAssignNew(NumTilesText, STextBlock)
+									SAssignNew(FormatText, STextBlock)
 								]
+
 							// Start Timecode.
 							+ SVerticalBox::Slot()
 								.AutoHeight()
@@ -128,6 +129,7 @@ void SMediaPlayerEditorMediaDetails::Construct(const FArguments& InArgs,
 								[
 									SAssignNew(StartTimecodeText, STextBlock)
 								]
+
 							// Seek Performance
 							+ SVerticalBox::Slot()
 								.AutoHeight()
@@ -136,6 +138,24 @@ void SMediaPlayerEditorMediaDetails::Construct(const FArguments& InArgs,
 								[
 									SAssignNew(SeekPerformance, STextBlock)
 								]
+
+							// Video decoder provider
+							+ SVerticalBox::Slot()
+								.AutoHeight()
+								.VAlign(VAlign_Center)
+								.Padding(4.0f)
+								[
+									SAssignNew(VideoDecoderProvider, STextBlock)
+								]
+
+							// Audio decoder provider
+							+ SVerticalBox::Slot()
+								.AutoHeight()
+								.VAlign(VAlign_Center)
+								.Padding(4.0f)
+								[
+									SAssignNew(AudioDecoderProvider, STextBlock)
+								]
 						]
 				]
 
@@ -167,6 +187,8 @@ void SMediaPlayerEditorMediaDetails::UpdateDetails()
 	int32 SurfaceHeight = 0;
 	FString StartTimecode;
 	int32 KeyframeInterval = -1;
+	FString VideoDecoderImpl;
+	FString AudioDecoderImpl;
 
 	// Get player info.
 	if (MediaPlayer != nullptr)
@@ -179,6 +201,8 @@ void SMediaPlayerEditorMediaDetails::UpdateDetails()
 		NumTotalTiles = NumTiles.X * NumTiles.Y;
 		MediaPlayer->GetMediaInfo<FString>(StartTimecode, UMediaPlayer::MediaInfoNameStartTimecodeValue.Resolve());
 		MediaPlayer->GetMediaInfo<int32>(KeyframeInterval, UMediaPlayer::MediaInfoNameKeyframeInterval.Resolve());
+		MediaPlayer->GetMediaInfo<FString>(VideoDecoderImpl, UMediaPlayer::MediaInfoNameVideoDecoderProvider.Resolve());
+		MediaPlayer->GetMediaInfo<FString>(AudioDecoderImpl, UMediaPlayer::MediaInfoNameAudioDecoderProvider.Resolve());
 	}
 
 	// Get texture info,
@@ -236,6 +260,25 @@ void SMediaPlayerEditorMediaDetails::UpdateDetails()
 	{
 		SeekPerformance->SetVisibility(EVisibility::Collapsed);
 	}
+
+	if (!VideoDecoderImpl.IsEmpty())
+	{
+		VideoDecoderProvider->SetText(FText::Format(LOCTEXT("VideoDecoderProvider", "Video decoder: {0}"), FText::FromString(VideoDecoderImpl)));
+		VideoDecoderProvider->SetVisibility(EVisibility::Visible);
+	}
+	else
+	{
+		VideoDecoderProvider->SetVisibility(EVisibility::Collapsed);
+	}
+	if (!AudioDecoderImpl.IsEmpty())
+	{
+		AudioDecoderProvider->SetText(FText::Format(LOCTEXT("AudioDecoderProvider", "Audio decoder: {0}"), FText::FromString(AudioDecoderImpl)));
+		AudioDecoderProvider->SetVisibility(EVisibility::Visible);
+	}
+	else
+	{
+		AudioDecoderProvider->SetVisibility(EVisibility::Collapsed);
+	}
 }
 
 #undef LOCTEXT_NAMESPACE
diff --git a/Engine/Plugins/Media/MediaPlayerEditor/Source/MediaPlayerEditor/Public/Widgets/SMediaPlayerEditorMediaDetails.h b/Engine/Plugins/Media/MediaPlayerEditor/Source/MediaPlayerEditor/Public/Widgets/SMediaPlayerEditorMediaDetails.h
--- a/Engine/Plugins/Media/MediaPlayerEditor/Source/MediaPlayerEditor/Public/Widgets/SMediaPlayerEditorMediaDetails.h
+++ b/Engine/Plugins/Media/MediaPlayerEditor/Source/MediaPlayerEditor/Public/Widgets/SMediaPlayerEditorMediaDetails.h
@@ -60,4 +60,6 @@ private:
 	TSharedPtr<STextBlock> ResourceSizeText;
 	TSharedPtr<STextBlock> StartTimecodeText;
 	TSharedPtr<STextBlock> SeekPerformance;
+	TSharedPtr<STextBlock> VideoDecoderProvider;
+	TSharedPtr<STextBlock> AudioDecoderProvider;
 };
diff --git a/Engine/Source/Runtime/MediaAssets/Private/Assets/MediaPlayer.cpp b/Engine/Source/Runtime/MediaAssets/Private/Assets/MediaPlayer.cpp
--- a/Engine/Source/Runtime/MediaAssets/Private/Assets/MediaPlayer.cpp
+++ b/Engine/Source/Runtime/MediaAssets/Private/Assets/MediaPlayer.cpp
@@ -38,6 +38,8 @@ FLazyName UMediaPlayer::MediaInfoNameSourceNumTiles(TEXT("SourceNumTiles"));
 FLazyName UMediaPlayer::MediaInfoNameStartTimecodeValue(TEXT("StartTimecodeValue"));
 FLazyName UMediaPlayer::MediaInfoNameStartTimecodeFrameRate(TEXT("StartTimecodeFrameRate"));
 FLazyName UMediaPlayer::MediaInfoNameKeyframeInterval(TEXT("KeyframeInterval"));
+FLazyName UMediaPlayer::MediaInfoNameVideoDecoderProvider(TEXT("VideoDecoderProvider"));
+FLazyName UMediaPlayer::MediaInfoNameAudioDecoderProvider(TEXT("AudioDecoderProvider"));
 
 
 /* UMediaPlayer structors
diff --git a/Engine/Source/Runtime/MediaAssets/Public/MediaPlayer.h b/Engine/Source/Runtime/MediaAssets/Public/MediaPlayer.h
--- a/Engine/Source/Runtime/MediaAssets/Public/MediaPlayer.h
+++ b/Engine/Source/Runtime/MediaAssets/Public/MediaPlayer.h
@@ -1330,6 +1330,10 @@ public:
 	static MEDIAASSETS_API FLazyName MediaInfoNameStartTimecodeFrameRate;
 	/** Keyframe interval. Int32. */
 	static MEDIAASSETS_API FLazyName MediaInfoNameKeyframeInterval;
+	/** Video decoder provider. String. */
+	static MEDIAASSETS_API FLazyName MediaInfoNameVideoDecoderProvider;
+	/** Audio decoder provider. String. */
+	static MEDIAASSETS_API FLazyName MediaInfoNameAudioDecoderProvider;
 
 	/**
 	 * Duration of samples to cache ahead of the play head.
