Media Player Editor: 360 view controls and viewport overlay have been removed.

#Synced-CL 50687911

diff --git a/Engine/Plugins/Media/MediaPlayerEditor/Source/MediaPlayerEditor/Private/Widgets/SMediaPlayerEditorDetails.cpp b/Engine/Plugins/Media/MediaPlayerEditor/Source/MediaPlayerEditor/Private/Widgets/SMediaPlayerEditorDetails.cpp
--- a/Engine/Plugins/Media/MediaPlayerEditor/Source/MediaPlayerEditor/Private/Widgets/SMediaPlayerEditorDetails.cpp
+++ b/Engine/Plugins/Media/MediaPlayerEditor/Source/MediaPlayerEditor/Private/Widgets/SMediaPlayerEditorDetails.cpp
@@ -30,6 +30,21 @@ void SMediaPlayerEditorDetails::Construct(const FArguments& InArgs, UMediaPlayer
 	}
 
 	TSharedPtr<IDetailsView> DetailsView = FModuleManager::GetModuleChecked<FPropertyEditorModule>("PropertyEditor").CreateDetailView(DetailsViewArgs);
+
+	DetailsView->SetIsPropertyVisibleDelegate(
+		FIsPropertyVisible::CreateLambda(
+			[](const FPropertyAndParent& InPropertyAndParent) -> bool
+			{
+				static const FLazyName HorizontalFieldOfViewName = "HorizontalFieldOfView";
+				static const FLazyName VerticalFieldOfViewName = "VerticalFieldOfView";
+				static const FLazyName ViewRotationName = "ViewRotation";
+
+				return (InPropertyAndParent.Property.GetFName() != HorizontalFieldOfViewName
+					&& InPropertyAndParent.Property.GetFName() != VerticalFieldOfViewName
+					&& InPropertyAndParent.Property.GetFName() != ViewRotationName);
+			}
+		));
+
 	DetailsView->SetObject(MediaPlayer);
 
 	ChildSlot
diff --git a/Engine/Plugins/Media/MediaPlayerEditor/Source/MediaPlayerEditor/Private/Widgets/SMediaPlayerEditorViewport.cpp b/Engine/Plugins/Media/MediaPlayerEditor/Source/MediaPlayerEditor/Private/Widgets/SMediaPlayerEditorViewport.cpp
--- a/Engine/Plugins/Media/MediaPlayerEditor/Source/MediaPlayerEditor/Private/Widgets/SMediaPlayerEditorViewport.cpp
+++ b/Engine/Plugins/Media/MediaPlayerEditor/Source/MediaPlayerEditor/Private/Widgets/SMediaPlayerEditorViewport.cpp
@@ -135,19 +135,6 @@ void SMediaPlayerEditorViewport::Construct(const FArguments& InArgs, UMediaPlaye
 									.Text(this, &SMediaPlayerEditorViewport::HandleMediaPlayerStateText)
 									.ToolTipText(LOCTEXT("OverlayPlayerStateTooltip", "The media player's current state"))
 							]
-
-						+ SHorizontalBox::Slot()
-							.HAlign(HAlign_Right)
-							.VAlign(VAlign_Bottom)
-							[
-								// view settings
-								SNew(STextBlock)
-									.ColorAndOpacity(FSlateColor::UseSubduedForeground())
-									.Font(OverlayFont)
-									.ShadowOffset(FVector2D(1.f, 1.f))
-									.Text(this, &SMediaPlayerEditorViewport::HandleViewSettingsText)
-									.ToolTipText(LOCTEXT("OverlayViewSettingsTooltip", "The current view settings"))
-							]
 					]
 			]
 
@@ -361,27 +348,4 @@ FText SMediaPlayerEditorViewport::HandleNotificationText() const
 }
 
 
-FText SMediaPlayerEditorViewport::HandleViewSettingsText() const
-{
-	FVector Euler = MediaPlayer->GetViewRotation().Euler();
-
-	FNumberFormattingOptions NumberFormat;
-	{
-		NumberFormat.MaximumFractionalDigits = 0;
-		NumberFormat.MinimumIntegralDigits = 3;
-	}
-
-	FFormatNamedArguments Args;
-	{
-		Args.Add(TEXT("P"), FText::AsNumber(Euler.Y, &NumberFormat));
-		Args.Add(TEXT("Y"), FText::AsNumber(Euler.Z, &NumberFormat));
-		Args.Add(TEXT("R"), FText::AsNumber(Euler.X, &NumberFormat));
-		Args.Add(TEXT("H"), FText::AsNumber(MediaPlayer->GetHorizontalFieldOfView(), &NumberFormat));
-		Args.Add(TEXT("V"), FText::AsNumber(MediaPlayer->GetVerticalFieldOfView(), &NumberFormat));
-	}
-
-	return FText::Format(LOCTEXT("ViewSettingsFormat", "R {P} {Y} {R} | FOV {H} {V}"), Args);
-}
-
-
 #undef LOCTEXT_NAMESPACE
diff --git a/Engine/Plugins/Media/MediaPlayerEditor/Source/MediaPlayerEditor/Private/Widgets/SMediaPlayerEditorViewport.h b/Engine/Plugins/Media/MediaPlayerEditor/Source/MediaPlayerEditor/Private/Widgets/SMediaPlayerEditorViewport.h
--- a/Engine/Plugins/Media/MediaPlayerEditor/Source/MediaPlayerEditor/Private/Widgets/SMediaPlayerEditorViewport.h
+++ b/Engine/Plugins/Media/MediaPlayerEditor/Source/MediaPlayerEditor/Private/Widgets/SMediaPlayerEditorViewport.h
@@ -68,9 +68,6 @@ private:
 	/** Callback for getting the text of the notification overlay. */
 	FText HandleNotificationText() const;
 
-	/** Callback for getting the text of the view settings overlay. */
-	FText HandleViewSettingsText() const;
-
 private:
 
 	/** Pointer to the media player that is being viewed. */
