Media Player Editor: Removes the 'name' column from the player details tab for platform-specific players.

#Synced-CL 50687911

diff --git a/Engine/Plugins/Media/MediaPlayerEditor/Source/MediaPlayerEditor/Private/Customizations/BaseMediaSourceCustomization.cpp b/Engine/Plugins/Media/MediaPlayerEditor/Source/MediaPlayerEditor/Private/Customizations/BaseMediaSourceCustomization.cpp
--- a/Engine/Plugins/Media/MediaPlayerEditor/Source/MediaPlayerEditor/Private/Customizations/BaseMediaSourceCustomization.cpp
+++ b/Engine/Plugins/Media/MediaPlayerEditor/Source/MediaPlayerEditor/Private/Customizations/BaseMediaSourceCustomization.cpp
@@ -45,11 +45,7 @@ void FBaseMediaSourceCustomization::CustomizeDetails(IDetailLayoutBuilder& Detai
 			PlayerNamesRow
 				.ShowPropertyButtons(false)
 				.CustomWidget()
-				.NameContent()
-				[
-					PlatformPlayerNamesProperty->CreatePropertyNameWidget()
-				]
-				.ValueContent()
+				.WholeRowContent()
 				.MaxDesiredWidth(0.0f)
 				[
 					MakePlatformPlayerNamesValueWidget()
