Media Player Editor: There are now 2 options for media backgrounds, black and checkered. Defaulted to black.

#Synced-CL 50687911

#Synced-CL 50687911

diff --git a/Engine/Plugins/Media/MediaPlayerEditor/Source/MediaPlayerEditor/Private/Shared/MediaPlayerEditorSettings.h b/Engine/Plugins/Media/MediaPlayerEditor/Source/MediaPlayerEditor/Private/Shared/MediaPlayerEditorSettings.h
--- a/Engine/Plugins/Media/MediaPlayerEditor/Source/MediaPlayerEditor/Private/Shared/MediaPlayerEditorSettings.h
+++ b/Engine/Plugins/Media/MediaPlayerEditor/Source/MediaPlayerEditor/Private/Shared/MediaPlayerEditorSettings.h
@@ -22,6 +22,13 @@ enum class EMediaPlayerEditorScale : uint8
 	Original
 };
 
+/* Background types */
+UENUM()
+enum class EMediaPlayerEditorBackground : uint8
+{
+	Black,
+	Checkered
+};
 
 UCLASS(config=EditorPerProjectUserSettings)
 class UMediaPlayerEditorSettings
@@ -42,4 +49,8 @@ public:
 	/** How the video viewport should be scaled. */
 	UPROPERTY(config, EditAnywhere, Category=Viewer)
 	EMediaPlayerEditorScale ViewportScale;
+
+	/** Which background to use on the video background. */
+	UPROPERTY(config, EditAnywhere, Category = Viewer)
+	EMediaPlayerEditorBackground MediaBackground = EMediaPlayerEditorBackground::Black;
 };
diff --git a/Engine/Plugins/Media/MediaPlayerEditor/Source/MediaPlayerEditor/Private/Widgets/SMediaPlayerEditorViewer.cpp b/Engine/Plugins/Media/MediaPlayerEditor/Source/MediaPlayerEditor/Private/Widgets/SMediaPlayerEditorViewer.cpp
--- a/Engine/Plugins/Media/MediaPlayerEditor/Source/MediaPlayerEditor/Private/Widgets/SMediaPlayerEditorViewer.cpp
+++ b/Engine/Plugins/Media/MediaPlayerEditor/Source/MediaPlayerEditor/Private/Widgets/SMediaPlayerEditorViewer.cpp
@@ -225,6 +225,14 @@ void SMediaPlayerEditorViewer::Construct(const FArguments& InArgs, UMediaPlayer&
 				FSlateIcon()
 			);
 
+			SettingsMenuBuilder.AddSubMenu(
+				LOCTEXT("BackgroundMenuLabel", "Background"),
+				LOCTEXT("BackgroundMenuTooltip", "Change the background behind the media"),
+				FNewMenuDelegate::CreateRaw(this, &SMediaPlayerEditorViewer::HandleBackgroundMenuNewMenu),
+				false,
+				FSlateIcon()
+			);
+
 			SettingsMenuBuilder.AddMenuEntry(
 				LOCTEXT("ShowTextOverlaysMenuLabel", "Show Text Overlays"),
 				LOCTEXT("ShowTextOverlaysMenuTooltip", "Show caption and subtitle text overlays"),
@@ -947,6 +955,55 @@ FText SMediaPlayerEditorViewer::HandleTimerTextBlockToolTipText() const
 	return FText::Format(LOCTEXT("TimeRemainingTooltipFormat", "{0} remaining"), FText::AsTimespan(Remaining));
 }
 
+void SMediaPlayerEditorViewer::HandleBackgroundMenuNewMenu(FMenuBuilder& MenuBuilder)
+{
+	MenuBuilder.AddMenuEntry(
+		LOCTEXT("BackgroundBlack", "Black"),
+		LOCTEXT("BackgroundBlackTooltip", "Use a black media background"),
+		FSlateIcon(),
+		FUIAction(
+			FExecuteAction::CreateLambda(
+				[] 
+				{
+					UMediaPlayerEditorSettings* Settings = GetMutableDefault<UMediaPlayerEditorSettings>();
+					Settings->MediaBackground = EMediaPlayerEditorBackground::Black;
+					Settings->SaveConfig();
+				}),
+			FCanExecuteAction(),
+			FIsActionChecked::CreateLambda(
+				[] 
+				{ 
+					return GetDefault<UMediaPlayerEditorSettings>()->MediaBackground == EMediaPlayerEditorBackground::Black;
+				})
+		),
+		NAME_None,
+		EUserInterfaceActionType::RadioButton
+	);
+
+	MenuBuilder.AddMenuEntry(
+		LOCTEXT("BackgroundCheckered", "Checkered"),
+		LOCTEXT("BackgroundCheckeredTooltip", "Use a checkered media background"),
+		FSlateIcon(),
+		FUIAction(
+			FExecuteAction::CreateLambda(
+				[] 
+				{
+					UMediaPlayerEditorSettings* Settings = GetMutableDefault<UMediaPlayerEditorSettings>();
+					Settings->MediaBackground = EMediaPlayerEditorBackground::Checkered;
+					Settings->SaveConfig();
+				}),
+			FCanExecuteAction(),
+			FIsActionChecked::CreateLambda(
+				[] 
+				{ 
+					return GetDefault<UMediaPlayerEditorSettings>()->MediaBackground == EMediaPlayerEditorBackground::Checkered;
+				})
+		),
+		NAME_None,
+		EUserInterfaceActionType::RadioButton
+	);
+}
+
 
 void SMediaPlayerEditorViewer::HandleTrackMenuNewMenu(FMenuBuilder& MenuBuilder, EMediaPlayerTrack TrackType)
 {
diff --git a/Engine/Plugins/Media/MediaPlayerEditor/Source/MediaPlayerEditor/Private/Widgets/SMediaPlayerEditorViewport.cpp b/Engine/Plugins/Media/MediaPlayerEditor/Source/MediaPlayerEditor/Private/Widgets/SMediaPlayerEditorViewport.cpp
--- a/Engine/Plugins/Media/MediaPlayerEditor/Source/MediaPlayerEditor/Private/Widgets/SMediaPlayerEditorViewport.cpp
+++ b/Engine/Plugins/Media/MediaPlayerEditor/Source/MediaPlayerEditor/Private/Widgets/SMediaPlayerEditorViewport.cpp
@@ -46,7 +46,7 @@ void SMediaPlayerEditorViewport::Construct(const FArguments& InArgs, UMediaPlaye
 		+ SOverlay::Slot()
 			[
 				SNew(SBorder)
-				.BorderImage(FCoreStyle::Get().GetBrush("ColorPicker.AlphaBackground"))
+				.BorderImage(this, &SMediaPlayerEditorViewport::GetMediaBorderImage)
 				.Padding(0.0f)
 				[
 					// movie viewport
@@ -347,5 +347,23 @@ FText SMediaPlayerEditorViewport::HandleNotificationText() const
 	return FText::GetEmpty();
 }
 
+const FSlateBrush* SMediaPlayerEditorViewport::GetMediaBorderImage() const
+{
+	if (const UMediaPlayerEditorSettings* Settings = GetDefault<UMediaPlayerEditorSettings>())
+	{
+		switch (Settings->MediaBackground)
+		{
+			case EMediaPlayerEditorBackground::Checkered:
+				return FCoreStyle::Get().GetBrush("ColorPicker.AlphaBackground");
+
+			default:
+				// Fall back to default black background
+				break;
+		}
+	}
+
+	return FCoreStyle::Get().GetBrush("BlackBrush");
+}
+
 
 #undef LOCTEXT_NAMESPACE
diff --git a/Engine/Plugins/Media/MediaPlayerEditor/Source/MediaPlayerEditor/Private/Widgets/SMediaPlayerEditorViewport.h b/Engine/Plugins/Media/MediaPlayerEditor/Source/MediaPlayerEditor/Private/Widgets/SMediaPlayerEditorViewport.h
--- a/Engine/Plugins/Media/MediaPlayerEditor/Source/MediaPlayerEditor/Private/Widgets/SMediaPlayerEditorViewport.h
+++ b/Engine/Plugins/Media/MediaPlayerEditor/Source/MediaPlayerEditor/Private/Widgets/SMediaPlayerEditorViewport.h
@@ -9,7 +9,7 @@
 class UMediaPlayer;
 class UMediaSoundComponent;
 class UMediaTexture;
-
+struct FSlateBrush;
 
 class SMediaPlayerEditorViewport
 	: public SCompoundWidget
@@ -68,6 +68,9 @@ private:
 	/** Callback for getting the text of the notification overlay. */
 	FText HandleNotificationText() const;
 
+	/** Returns the image used for the background of the media. */
+	const FSlateBrush* GetMediaBorderImage() const;
+
 private:
 
 	/** Pointer to the media player that is being viewed. */
diff --git a/Engine/Plugins/Media/MediaPlayerEditor/Source/MediaPlayerEditor/Public/Widgets/SMediaPlayerEditorViewer.h b/Engine/Plugins/Media/MediaPlayerEditor/Source/MediaPlayerEditor/Public/Widgets/SMediaPlayerEditorViewer.h
--- a/Engine/Plugins/Media/MediaPlayerEditor/Source/MediaPlayerEditor/Public/Widgets/SMediaPlayerEditorViewer.h
+++ b/Engine/Plugins/Media/MediaPlayerEditor/Source/MediaPlayerEditor/Public/Widgets/SMediaPlayerEditorViewer.h
@@ -120,6 +120,9 @@ private:
 	/** Callback for getting the tool tip of the timer text block. */
 	UE_API FText HandleTimerTextBlockToolTipText() const;
 
+	/** Callback for creating the Background sub-menu. */
+	UE_API void HandleBackgroundMenuNewMenu(FMenuBuilder& MenuBuilder);
+
 	/** Callback for creating a track sub-menu. */
 	UE_API void HandleTrackMenuNewMenu(FMenuBuilder& MenuBuilder, EMediaPlayerTrack TrackType);
 
