Media Player Editor: Added keybinds to the playback tab:
- Space / K: Toggle play/pause
- J: Jump back 10s
- Ctrl+J: Jump to the beginning
- Alt+J: Skip back 1 frame (NYI)
- L: Jump forward 10s
- Ctrl+L: Jump to the end (will reset to the beginning if looping is on)
- Alt+L: Skip forward 1 frame (NYI)
- Escape: Close media

The video tab must be focused.

#Synced-CL 50687911

diff --git a/Engine/Plugins/Media/MediaPlate/Source/MediaPlate/Private/MediaPlateComponent.cpp b/Engine/Plugins/Media/MediaPlate/Source/MediaPlate/Private/MediaPlateComponent.cpp
--- a/Engine/Plugins/Media/MediaPlate/Source/MediaPlate/Private/MediaPlateComponent.cpp
+++ b/Engine/Plugins/Media/MediaPlate/Source/MediaPlate/Private/MediaPlateComponent.cpp
@@ -443,7 +443,12 @@ bool UMediaPlateComponent::IsEventStateChangeAllowed(EMediaPlateEventState InEve
 	{
 	case EMediaPlateEventState::Open:
 	case EMediaPlateEventState::Play:
+	case EMediaPlateEventState::JumpForward:
+	case EMediaPlateEventState::JumpBackward:
 		return !IsAnyMediaTextureProxied();
+	case EMediaPlateEventState::StepForward:
+	case EMediaPlateEventState::StepBackward:
+		return IntendedPlaybackState == EPlaybackState::Paused && !IsAnyMediaTextureProxied();
 	default:
 		return true;
 	}
@@ -1836,6 +1841,8 @@ void UMediaPlateComponent::PostEditUndo()
 
 void UMediaPlateComponent::SwitchStates(EMediaPlateEventState State)
 {
+	static const FTimespan TenSeconds = FTimespan::FromSeconds(10.f);
+
 	switch (State)
 	{
 	case EMediaPlateEventState::Play:
@@ -1892,6 +1899,77 @@ void UMediaPlateComponent::SwitchStates(EMediaPlateEventState State)
 			Previous();
 		}
 		break;
+	case EMediaPlateEventState::FastForward:
+		{
+			if (MediaPlayer != nullptr)
+			{
+				MediaPlayer->Seek(MediaPlayer->GetDuration());
+			}
+		}
+		break;
+	case EMediaPlateEventState::StepForward:
+		{
+			if (MediaPlayer != nullptr)
+			{
+				// TODO: Convert to underlying media player command when framework support is added.
+				const float FrameRate = MediaPlayer->GetVideoTrackFrameRate(INDEX_NONE, INDEX_NONE);
+
+				if (!FMath::IsNearlyZero(FrameRate))
+				{
+					MediaPlayer->Seek(MediaPlayer->GetTime() + FTimespan::FromSeconds(1.f / FrameRate));
+				}
+			}
+		}
+		break;
+	case EMediaPlateEventState::StepBackward:
+		{
+			if (MediaPlayer != nullptr)
+			{
+				// TODO: Convert to underlying media player command when framework support is added.
+				const float FrameRate = MediaPlayer->GetVideoTrackFrameRate(INDEX_NONE, INDEX_NONE);
+
+				if (!FMath::IsNearlyZero(FrameRate))
+				{
+					MediaPlayer->Seek(MediaPlayer->GetTime() - FTimespan::FromSeconds(1.f / FrameRate));
+				}
+			}
+		}
+		break;
+	case EMediaPlateEventState::JumpForward:
+		{
+			if (MediaPlayer != nullptr)
+			{
+				const FTimespan CurrentTime = MediaPlayer->GetTime();
+				const FTimespan Duration = MediaPlayer->GetDuration();
+
+				if ((Duration - CurrentTime) <= TenSeconds)
+				{
+					MediaPlayer->Seek(Duration);
+				}
+				else
+				{
+					MediaPlayer->Seek(MediaPlayer->GetTime() + FTimespan::FromSeconds(10.f));
+				}
+			}
+		}
+		break;
+	case EMediaPlateEventState::JumpBackward:
+		{
+			if (MediaPlayer != nullptr)
+			{
+				const FTimespan CurrentTime = MediaPlayer->GetTime();
+
+				if (CurrentTime <= TenSeconds)
+				{
+					MediaPlayer->Seek(0);
+				}
+				else
+				{
+					MediaPlayer->Seek(MediaPlayer->GetTime() - FTimespan::FromSeconds(10.f));
+				}
+			}
+		}
+		break;
 	case EMediaPlateEventState::MAX:
 	default:
 		checkNoEntry();
diff --git a/Engine/Plugins/Media/MediaPlate/Source/MediaPlate/Public/MediaPlateComponent.h b/Engine/Plugins/Media/MediaPlate/Source/MediaPlate/Public/MediaPlateComponent.h
--- a/Engine/Plugins/Media/MediaPlate/Source/MediaPlate/Public/MediaPlateComponent.h
+++ b/Engine/Plugins/Media/MediaPlate/Source/MediaPlate/Public/MediaPlateComponent.h
@@ -38,6 +38,11 @@ enum class EMediaPlateEventState : uint8
 	Rewind,
 	Next,
 	Previous,
+	FastForward,
+	StepForward,
+	StepBackward,
+	JumpForward,
+	JumpBackward,
 	MAX
 };
 
diff --git a/Engine/Plugins/Media/MediaPlate/Source/MediaPlateEditor/Private/MediaPlateEditorModule.cpp b/Engine/Plugins/Media/MediaPlate/Source/MediaPlateEditor/Private/MediaPlateEditorModule.cpp
--- a/Engine/Plugins/Media/MediaPlate/Source/MediaPlateEditor/Private/MediaPlateEditorModule.cpp
+++ b/Engine/Plugins/Media/MediaPlate/Source/MediaPlateEditor/Private/MediaPlateEditorModule.cpp
@@ -18,7 +18,6 @@
 #include "MediaSoundComponent.h"
 #include "MediaSource.h"
 #include "MediaTexture.h"
-#include "Models/MediaPlateEditorCommands.h"
 #include "Modules/ModuleManager.h"
 #include "PropertyEditorModule.h"
 #include "Sequencer/MediaPlateTrackEditor.h"
@@ -36,7 +35,6 @@ void FMediaPlateEditorModule::StartupModule()
 {
 	FMediaPlateEditorStyle::Get();
 
-	FMediaPlateEditorCommands::Register();
 	RegisterPlacementModeItems();
 	RegisterSectionMappings();
 
diff --git a/Engine/Plugins/Media/MediaPlate/Source/MediaPlateEditor/Private/Models/MediaPlateEditorCommands.cpp b/Engine/Plugins/Media/MediaPlate/Source/MediaPlateEditor/Private/Models/MediaPlateEditorCommands.cpp
deleted file mode 100644
diff --git a/Engine/Plugins/Media/MediaPlate/Source/MediaPlateEditor/Private/Models/MediaPlateEditorCommands.h b/Engine/Plugins/Media/MediaPlate/Source/MediaPlateEditor/Private/Models/MediaPlateEditorCommands.h
deleted file mode 100644
diff --git a/Engine/Plugins/Media/MediaPlate/Source/MediaPlateEditor/Private/Toolkits/MediaPlateEditorToolkit.cpp b/Engine/Plugins/Media/MediaPlate/Source/MediaPlateEditor/Private/Toolkits/MediaPlateEditorToolkit.cpp
--- a/Engine/Plugins/Media/MediaPlate/Source/MediaPlateEditor/Private/Toolkits/MediaPlateEditorToolkit.cpp
+++ b/Engine/Plugins/Media/MediaPlate/Source/MediaPlateEditor/Private/Toolkits/MediaPlateEditorToolkit.cpp
@@ -13,7 +13,7 @@
 #include "MediaPlateEditorModule.h"
 #include "MediaPlayer.h"
 #include "MediaPlaylist.h"
-#include "Models/MediaPlateEditorCommands.h"
+#include "Models/MediaPlayerEditorCommands.h"
 #include "SlateOptMacros.h"
 #include "Styling/AppStyle.h"
 #include "Widgets/Docking/SDockTab.h"
@@ -259,7 +259,7 @@ void FMediaPlateEditorToolkit::PostRedo(bool bSuccess)
 
 void FMediaPlateEditorToolkit::BindCommands()
 {
-	const FMediaPlateEditorCommands& Commands = FMediaPlateEditorCommands::Get();
+	const FMediaPlayerEditorCommands& Commands = FMediaPlayerEditorCommands::GetExternal();
 
 	ToolkitCommands->MapAction(
 		Commands.CloseMedia,
@@ -311,8 +311,44 @@ void FMediaPlateEditorToolkit::BindCommands()
 
 	ToolkitCommands->MapAction(
 		Commands.RewindMedia,
-		FExecuteAction::CreateSP( this, &FMediaPlateEditorToolkit::HandleMediaPlateEvent, EMediaPlateEventState::Rewind ),
-		FCanExecuteAction::CreateSP( this, &FMediaPlateEditorToolkit::CanHandleMediaPlateEvent, EMediaPlateEventState::Rewind )
+		FExecuteAction::CreateSP(this, &FMediaPlateEditorToolkit::HandleMediaPlateEvent, EMediaPlateEventState::Rewind),
+		FCanExecuteAction::CreateSP(this, &FMediaPlateEditorToolkit::CanHandleMediaPlateEvent, EMediaPlateEventState::Rewind)
+	);
+
+	ToolkitCommands->MapAction(
+		Commands.TogglePlayPauseMedia,
+		FExecuteAction::CreateSP( this, &FMediaPlateEditorToolkit::HandleMediaPlateTogglePlayPauseEvent ),
+		FCanExecuteAction::CreateSP( this, &FMediaPlateEditorToolkit::CanHandleMediaPlateTogglePlayPauseEvent )
+	);
+
+	ToolkitCommands->MapAction(
+		Commands.FastForwardMedia,
+		FExecuteAction::CreateSP(this, &FMediaPlateEditorToolkit::HandleMediaPlateEvent, EMediaPlateEventState::FastForward),
+		FCanExecuteAction::CreateSP(this, &FMediaPlateEditorToolkit::CanHandleMediaPlateEvent, EMediaPlateEventState::FastForward)
+	);
+
+	ToolkitCommands->MapAction(
+		Commands.StepForward,
+		FExecuteAction::CreateSP(this, &FMediaPlateEditorToolkit::HandleMediaPlateEvent, EMediaPlateEventState::StepForward),
+		FCanExecuteAction::CreateSP(this, &FMediaPlateEditorToolkit::CanHandleMediaPlateEvent, EMediaPlateEventState::StepForward)
+	);
+
+	ToolkitCommands->MapAction(
+		Commands.StepBackward,
+		FExecuteAction::CreateSP(this, &FMediaPlateEditorToolkit::HandleMediaPlateEvent, EMediaPlateEventState::StepBackward),
+		FCanExecuteAction::CreateSP(this, &FMediaPlateEditorToolkit::CanHandleMediaPlateEvent, EMediaPlateEventState::StepBackward)
+	);
+
+	ToolkitCommands->MapAction(
+		Commands.JumpForward,
+		FExecuteAction::CreateSP(this, &FMediaPlateEditorToolkit::HandleMediaPlateEvent, EMediaPlateEventState::JumpForward),
+		FCanExecuteAction::CreateSP(this, &FMediaPlateEditorToolkit::CanHandleMediaPlateEvent, EMediaPlateEventState::JumpForward)
+	);
+
+	ToolkitCommands->MapAction(
+		Commands.JumpBackward,
+		FExecuteAction::CreateSP(this, &FMediaPlateEditorToolkit::HandleMediaPlateEvent, EMediaPlateEventState::JumpBackward),
+		FCanExecuteAction::CreateSP(this, &FMediaPlateEditorToolkit::CanHandleMediaPlateEvent, EMediaPlateEventState::JumpBackward)
 	);
 }
 
@@ -323,22 +359,24 @@ void FMediaPlateEditorToolkit::ExtendToolBar()
 	{
 		static void FillToolbar(FToolBarBuilder& ToolbarBuilder, const TSharedRef<FUICommandList> ToolkitCommands)
 		{
+			const FMediaPlayerEditorCommands& Commands = FMediaPlayerEditorCommands::GetExternal();
+
 			ToolbarBuilder.BeginSection("PlaybackControls");
 			{
-				ToolbarBuilder.AddToolBarButton(FMediaPlateEditorCommands::Get().PreviousMedia);
-				ToolbarBuilder.AddToolBarButton(FMediaPlateEditorCommands::Get().RewindMedia);
-				ToolbarBuilder.AddToolBarButton(FMediaPlateEditorCommands::Get().ReverseMedia);
-				ToolbarBuilder.AddToolBarButton(FMediaPlateEditorCommands::Get().PlayMedia);
-				ToolbarBuilder.AddToolBarButton(FMediaPlateEditorCommands::Get().PauseMedia);
-				ToolbarBuilder.AddToolBarButton(FMediaPlateEditorCommands::Get().ForwardMedia);
-				ToolbarBuilder.AddToolBarButton(FMediaPlateEditorCommands::Get().NextMedia);
+				ToolbarBuilder.AddToolBarButton(Commands.PreviousMedia);
+				ToolbarBuilder.AddToolBarButton(Commands.RewindMedia);
+				ToolbarBuilder.AddToolBarButton(Commands.ReverseMedia);
+				ToolbarBuilder.AddToolBarButton(Commands.PlayMedia);
+				ToolbarBuilder.AddToolBarButton(Commands.PauseMedia);
+				ToolbarBuilder.AddToolBarButton(Commands.ForwardMedia);
+				ToolbarBuilder.AddToolBarButton(Commands.NextMedia);
 			}
 			ToolbarBuilder.EndSection();
 
 			ToolbarBuilder.BeginSection("MediaControls");
 			{
-				ToolbarBuilder.AddToolBarButton(FMediaPlateEditorCommands::Get().OpenMedia);
-				ToolbarBuilder.AddToolBarButton(FMediaPlateEditorCommands::Get().CloseMedia);
+				ToolbarBuilder.AddToolBarButton(Commands.OpenMedia);
+				ToolbarBuilder.AddToolBarButton(Commands.CloseMedia);
 			}
 			ToolbarBuilder.EndSection();
 		}
@@ -412,6 +450,39 @@ void FMediaPlateEditorToolkit::HandleMediaPlateEvent(EMediaPlateEventState InSta
 	FMediaPlateCustomization::HandleMediaPlateEvent(MediaPlates, InState);
 }
 
+bool FMediaPlateEditorToolkit::CanHandleMediaPlateTogglePlayPauseEvent()
+{
+	if (UMediaPlayer* MediaPlayer = MediaPlate->GetMediaPlayer())
+	{
+		if (MediaPlayer->IsPlaying())
+		{
+			return CanHandleMediaPlateEvent(EMediaPlateEventState::Pause);
+		}
+
+		if (MediaPlayer->IsPaused())
+		{
+			return CanHandleMediaPlateEvent(EMediaPlateEventState::Play);
+		}
+	}
+
+	return false;
+}
+
+void FMediaPlateEditorToolkit::HandleMediaPlateTogglePlayPauseEvent()
+{
+	if (UMediaPlayer* MediaPlayer = MediaPlate->GetMediaPlayer())
+	{
+		if (MediaPlayer->IsPlaying())
+		{
+			HandleMediaPlateEvent(EMediaPlateEventState::Pause);
+		}
+		else if (MediaPlayer->IsPaused())
+		{
+			HandleMediaPlateEvent(EMediaPlateEventState::Play);
+		}
+	}
+}
+
 TSharedRef<SDockTab> FMediaPlateEditorToolkit::HandleTabManagerSpawnTab(const FSpawnTabArgs& Args, FName TabIdentifier)
 {
 	UMediaPlayer* MediaPlayer = MediaPlate->GetMediaPlayer();
@@ -434,7 +505,8 @@ TSharedRef<SDockTab> FMediaPlateEditorToolkit::HandleTabManagerSpawnTab(const FS
 	{
 		if (MediaPlayer != nullptr)
 		{
-			TabWidget = SNew(SMediaPlayerEditorViewer, *MediaPlayer, MediaTexture, Style, false);
+			TabWidget = SNew(SMediaPlayerEditorViewer, *MediaPlayer, MediaTexture, Style, false)
+				.Commands(ToolkitCommands);
 		}
 	}
 
diff --git a/Engine/Plugins/Media/MediaPlate/Source/MediaPlateEditor/Private/Toolkits/MediaPlateEditorToolkit.h b/Engine/Plugins/Media/MediaPlate/Source/MediaPlateEditor/Private/Toolkits/MediaPlateEditorToolkit.h
--- a/Engine/Plugins/Media/MediaPlate/Source/MediaPlateEditor/Private/Toolkits/MediaPlateEditorToolkit.h
+++ b/Engine/Plugins/Media/MediaPlate/Source/MediaPlateEditor/Private/Toolkits/MediaPlateEditorToolkit.h
@@ -90,6 +90,12 @@ private:
 
 	/** Callback to handle playback button events. */
 	void HandleMediaPlateEvent(EMediaPlateEventState InState);
+	
+	/** Callback specifically to validate the toggle between play and pause. */
+	bool CanHandleMediaPlateTogglePlayPauseEvent();
+
+	/** Callback specifically for toggle play/pause which then triggers a play or pause event. */
+	void HandleMediaPlateTogglePlayPauseEvent();
 
 	/** Callback for spawning tabs. */
 	TSharedRef<SDockTab> HandleTabManagerSpawnTab(const FSpawnTabArgs& Args, FName TabIdentifier);
diff --git a/Engine/Plugins/Media/MediaPlayerEditor/Source/MediaPlayerEditor/Private/Models/MediaPlayerEditorCommands.cpp b/Engine/Plugins/Media/MediaPlayerEditor/Source/MediaPlayerEditor/Private/Models/MediaPlayerEditorCommands.cpp
--- a/Engine/Plugins/Media/MediaPlayerEditor/Source/MediaPlayerEditor/Private/Models/MediaPlayerEditorCommands.cpp
+++ b/Engine/Plugins/Media/MediaPlayerEditor/Source/MediaPlayerEditor/Private/Models/MediaPlayerEditorCommands.cpp
@@ -5,19 +5,91 @@
 
 #define LOCTEXT_NAMESPACE "FMediaPlayerEditorCommands"
 
+const FMediaPlayerEditorCommands& FMediaPlayerEditorCommands::GetExternal()
+{
+	if (!IsRegistered())
+	{
+		Register();
+	}
+
+	return Get();
+}
 
 void FMediaPlayerEditorCommands::RegisterCommands()
 {
-	UI_COMMAND(CloseMedia, "Close", "Close the currently opened media", EUserInterfaceActionType::Button, FInputChord());
-	UI_COMMAND(ForwardMedia, "Forward", "Fast forward media playback", EUserInterfaceActionType::Button, FInputChord());
+	UI_COMMAND(CloseMedia, "Close", "Close the currently opened media", EUserInterfaceActionType::Button, FInputChord(EKeys::Escape));
 	UI_COMMAND(GenerateThumbnail, "Thumbnail", "Generate a thumbnail", EUserInterfaceActionType::Button, FInputChord());
-	UI_COMMAND(NextMedia, "Next", "Jump to next item in the play list", EUserInterfaceActionType::Button, FInputChord());
 	UI_COMMAND(OpenMedia, "Open", "Open the current media", EUserInterfaceActionType::Button, FInputChord());
 	UI_COMMAND(PauseMedia, "Pause", "Pause media playback", EUserInterfaceActionType::Button, FInputChord());
 	UI_COMMAND(PlayMedia, "Play", "Start media playback", EUserInterfaceActionType::Button, FInputChord());
+	UI_COMMAND(NextMedia, "Next", "Jump to next item in the play list", EUserInterfaceActionType::Button, FInputChord());
 	UI_COMMAND(PreviousMedia, "Prev", "Jump to previous item in the play list", EUserInterfaceActionType::Button, FInputChord());
 	UI_COMMAND(ReverseMedia, "Reverse", "Reverse media playback", EUserInterfaceActionType::Button, FInputChord());
-	UI_COMMAND(RewindMedia, "Rewind", "Rewind the media to the beginning", EUserInterfaceActionType::Button, FInputChord());
+	
+	UI_COMMAND(
+		RewindMedia,
+		"Rewind",
+		"Rewind the media to the beginning",
+		EUserInterfaceActionType::Button,
+		FInputChord(EKeys::J, EModifierKey::Control)
+	);
+
+	UI_COMMAND(
+		JumpBackward,
+		"Jump Backward",
+		"Jump the media backward 10 seconds",
+		EUserInterfaceActionType::Button,
+		FInputChord(EKeys::J)
+	);
+
+	UI_COMMAND(
+		StepBackward,
+		"Step Backward",
+		"Step the media backward 1 frame if it is paused",
+		EUserInterfaceActionType::Button,
+		FInputChord(EKeys::J, EModifierKey::Alt)
+	);
+
+	UI_COMMAND(
+		TogglePlayPauseMedia,
+		"Toggle Play/Pause",
+		"Toggles between media play and pause",
+		EUserInterfaceActionType::ToggleButton,
+		FInputChord(EKeys::SpaceBar),
+		FInputChord(EKeys::K)
+	);
+
+	UI_COMMAND(
+		ForwardMedia, 
+		"Forward", 
+		"Fast forward media playback", 
+		EUserInterfaceActionType::Button, 
+		FInputChord()
+	);
+
+	UI_COMMAND(
+		StepForward,
+		"Step Forward",
+		"Step the media forward 1 frame if it is paused",
+		EUserInterfaceActionType::Button,
+		FInputChord(EKeys::L, EModifierKey::Alt)
+	);
+
+	UI_COMMAND(
+		JumpForward,
+		"Jump Forward",
+		"Jump the media forward 10 seconds",
+		EUserInterfaceActionType::Button,
+		FInputChord(EKeys::L)
+	);
+
+	UI_COMMAND(
+		FastForwardMedia,
+		"Fast Forward", 
+		"Fast forwards the media to the end", 
+		EUserInterfaceActionType::Button, 
+		FInputChord(EKeys::L, EModifierKey::Control)
+	);
 }
 
 
diff --git a/Engine/Plugins/Media/MediaPlayerEditor/Source/MediaPlayerEditor/Private/Models/MediaPlayerEditorCommands.h b/Engine/Plugins/Media/MediaPlayerEditor/Source/MediaPlayerEditor/Private/Models/MediaPlayerEditorCommands.h
deleted file mode 100644
diff --git a/Engine/Plugins/Media/MediaPlayerEditor/Source/MediaPlayerEditor/Private/Toolkits/MediaPlayerEditorToolkit.cpp b/Engine/Plugins/Media/MediaPlayerEditor/Source/MediaPlayerEditor/Private/Toolkits/MediaPlayerEditorToolkit.cpp
--- a/Engine/Plugins/Media/MediaPlayerEditor/Source/MediaPlayerEditor/Private/Toolkits/MediaPlayerEditorToolkit.cpp
+++ b/Engine/Plugins/Media/MediaPlayerEditor/Source/MediaPlayerEditor/Private/Toolkits/MediaPlayerEditorToolkit.cpp
@@ -1,26 +1,22 @@
 // Copyright Epic Games, Inc. All Rights Reserved.
 
 #include "Toolkits/MediaPlayerEditorToolkit.h"
-#include "Framework/MultiBox/MultiBoxBuilder.h"
-#include "Editor.h"
+
+#include "Framework/Docking/TabManager.h"
+#include "Framework/MultiBox/MultiBoxExtender.h"
+#include "MediaPlayer.h"
+#include "Models/MediaPlayerEditorCommands.h"
 #include "SlateOptMacros.h"
-#include "Styling/AppStyle.h"
-#include "EditorReimportHandler.h"
-#include "MediaPlaylist.h"
 #include "Widgets/Docking/SDockTab.h"
-#include "Models/MediaPlayerEditorCommands.h"
 #include "Widgets/SMediaPlayerEditorDetails.h"
 #include "Widgets/SMediaPlayerEditorInfo.h"
 #include "Widgets/SMediaPlayerEditorMedia.h"
 #include "Widgets/SMediaPlayerEditorPlaylist.h"
 #include "Widgets/SMediaPlayerEditorStats.h"
-#include "MediaPlayer.h"
 #include "Widgets/SMediaPlayerEditorViewer.h"
 
-
 #define LOCTEXT_NAMESPACE "FMediaPlayerEditorToolkit"
 
-
 namespace MediaPlayerEditorToolkit
 {
 	static const FName AppIdentifier("MediaPlayerEditorApp");
@@ -32,27 +28,9 @@ namespace MediaPlayerEditorToolkit
 	static const FName ViewerTabId("Viewer");
 }
 
-
-/* FMediaPlayerEditorToolkit structors
- *****************************************************************************/
-
 FMediaPlayerEditorToolkit::FMediaPlayerEditorToolkit(const TSharedRef<ISlateStyle>& InStyle)
-	: MediaPlayer(nullptr)
-	, Style(InStyle)
-{ }
-
-
-FMediaPlayerEditorToolkit::~FMediaPlayerEditorToolkit()
-{
-	FReimportManager::Instance()->OnPreReimport().RemoveAll(this);
-	FReimportManager::Instance()->OnPostReimport().RemoveAll(this);
-
-	GEditor->UnregisterForUndo(this);
-}
-
-
-/* FMediaPlayerEditorToolkit interface
- *****************************************************************************/
+	: FMediaPlayerEditorToolkitMediaPlayerBase(InStyle)
+{}
 
 void FMediaPlayerEditorToolkit::Initialize(UMediaPlayer* InMediaPlayer, const EToolkitMode::Type InMode, const TSharedPtr<IToolkitHost>& InToolkitHost)
 {
@@ -65,12 +43,16 @@ void FMediaPlayerEditorToolkit::Initialize(UMediaPlayer* InMediaPlayer, const ET
 
 	// support undo/redo
 	MediaPlayer->SetFlags(RF_Transactional);
-	GEditor->RegisterForUndo(this);
 
-	BindCommands();
+	FMediaPlayerEditorToolkitMediaPlayerBase::Initialize(InMediaPlayer, MediaPlayerEditorToolkit::AppIdentifier, InMode, InToolkitHost);
+	
+//	IMediaPlayerEditorModule* MediaPlayerEditorModule = &FModuleManager::LoadModuleChecked<IMediaPlayerEditorModule>("MediaPlayerEditor");
+//	AddMenuExtender(MediaPlayerEditorModule->GetMenuExtensibilityManager()->GetAllExtenders(GetToolkitCommands(), GetEditingObjects()));
+}
 
-	// create tab layout
-	const TSharedRef<FTabManager::FLayout> Layout = FTabManager::NewLayout("Standalone_MediaPlayerEditor_v11")
+TSharedRef<FTabManager::FLayout> FMediaPlayerEditorToolkit::CreateLayout()
+{
+	return FTabManager::NewLayout("Standalone_MediaPlayerEditor_v11")
 		->AddArea
 		(
 			FTabManager::NewPrimaryArea()
@@ -120,34 +102,8 @@ void FMediaPlayerEditorToolkit::Initialize(UMediaPlayer* InMediaPlayer, const ET
 						)
 				)
 		);
-
-	FAssetEditorToolkit::InitAssetEditor(
-		InMode,
-		InToolkitHost,
-		MediaPlayerEditorToolkit::AppIdentifier,
-		Layout,
-		true /*bCreateDefaultStandaloneMenu*/,
-		true /*bCreateDefaultToolbar*/,
-		InMediaPlayer
-	);
-	
-//	IMediaPlayerEditorModule* MediaPlayerEditorModule = &FModuleManager::LoadModuleChecked<IMediaPlayerEditorModule>("MediaPlayerEditor");
-//	AddMenuExtender(MediaPlayerEditorModule->GetMenuExtensibilityManager()->GetAllExtenders(GetToolkitCommands(), GetEditingObjects()));
-
-	ExtendToolBar();
-	RegenerateMenusAndToolbars();
-}
-
-
-/* FAssetEditorToolkit interface
- *****************************************************************************/
-
-FString FMediaPlayerEditorToolkit::GetDocumentationLink() const
-{
-	return FString(TEXT("WorkingWithMedia/IntegratingMedia/MediaFramework"));
 }
 
-
 void FMediaPlayerEditorToolkit::OnClose()
 {
 	if (MediaPlayer != nullptr)
@@ -156,7 +112,6 @@ void FMediaPlayerEditorToolkit::OnClose()
 	}
 }
 
-
 void FMediaPlayerEditorToolkit::RegisterTabSpawners(const TSharedRef<FTabManager>& InTabManager)
 {
 	WorkspaceMenuCategory = InTabManager->AddLocalWorkspaceMenuCategory(LOCTEXT("WorkspaceMenu_MediaPlayerEditor", "Media Player Editor"));
@@ -195,7 +150,6 @@ void FMediaPlayerEditorToolkit::RegisterTabSpawners(const TSharedRef<FTabManager
 		.SetIcon(FSlateIcon(Style->GetStyleSetName(), "MediaPlayerEditor.Tabs.Stats"));
 }
 
-
 void FMediaPlayerEditorToolkit::UnregisterTabSpawners(const TSharedRef<class FTabManager>& InTabManager)
 {
 	FAssetEditorToolkit::UnregisterTabSpawners(InTabManager);
@@ -208,115 +162,21 @@ void FMediaPlayerEditorToolkit::UnregisterTabSpawners(const TSharedRef<class FTa
 	InTabManager->UnregisterTabSpawner(MediaPlayerEditorToolkit::DetailsTabId);
 }
 
-
-/* IToolkit interface
- *****************************************************************************/
-
 FText FMediaPlayerEditorToolkit::GetBaseToolkitName() const
 {
 	return LOCTEXT("AppLabel", "Media Player Editor");
 }
 
-
 FName FMediaPlayerEditorToolkit::GetToolkitFName() const
 {
 	return FName("MediaPlayerEditor");
 }
 
-
-FLinearColor FMediaPlayerEditorToolkit::GetWorldCentricTabColorScale() const
-{
-	return FLinearColor(0.3f, 0.2f, 0.5f, 0.5f);
-}
-
-
 FString FMediaPlayerEditorToolkit::GetWorldCentricTabPrefix() const
 {
 	return LOCTEXT("WorldCentricTabPrefix", "MediaPlayer ").ToString();
 }
 
-
-/* FGCObject interface
- *****************************************************************************/
-
-void FMediaPlayerEditorToolkit::AddReferencedObjects(FReferenceCollector& Collector)
-{
-	Collector.AddReferencedObject(MediaPlayer);
-}
-
-
-/* FEditorUndoClient interface
-*****************************************************************************/
-
-void FMediaPlayerEditorToolkit::PostUndo(bool bSuccess)
-{
-	// do nothing
-}
-
-
-void FMediaPlayerEditorToolkit::PostRedo(bool bSuccess)
-{
-	PostUndo(bSuccess);
-}
-
-
-/* FMediaPlayerEditorToolkit implementation
- *****************************************************************************/
-
-void FMediaPlayerEditorToolkit::BindCommands()
-{
-	const FMediaPlayerEditorCommands& Commands = FMediaPlayerEditorCommands::Get();
-
-	ToolkitCommands->MapAction(
-		Commands.CloseMedia,
-		FExecuteAction::CreateLambda([this] { MediaPlayer->Close(); }),
-		FCanExecuteAction::CreateLambda([this] { return !MediaPlayer->GetUrl().IsEmpty(); })
-	);
-
-	ToolkitCommands->MapAction(
-		Commands.ForwardMedia,
-		FExecuteAction::CreateLambda([this]{ MediaPlayer->SetRate(GetForwardRate()); }),
-		FCanExecuteAction::CreateLambda([this]{ return MediaPlayer->IsReady() && MediaPlayer->SupportsRate(GetForwardRate(), false); })
-	);
-
-	ToolkitCommands->MapAction(
-		Commands.NextMedia,
-		FExecuteAction::CreateLambda([this]{ MediaPlayer->Next(); }),
-		FCanExecuteAction::CreateLambda([this]{ return (MediaPlayer->GetPlaylistRef().Num() > 1); })
-	);
-
-	ToolkitCommands->MapAction(
-		Commands.PauseMedia,
-		FExecuteAction::CreateLambda([this]{ MediaPlayer->Pause(); }),
-		FCanExecuteAction::CreateLambda([this]{ return MediaPlayer->CanPause() && !MediaPlayer->IsPaused(); })
-	);
-
-	ToolkitCommands->MapAction(
-		Commands.PlayMedia,
-		FExecuteAction::CreateLambda([this]{ MediaPlayer->Play(); }),
-		FCanExecuteAction::CreateLambda([this]{ return MediaPlayer->IsReady() && (!MediaPlayer->IsPlaying() || (MediaPlayer->GetRate() != 1.0f)); })
-	);
-
-	ToolkitCommands->MapAction(
-		Commands.PreviousMedia,
-		FExecuteAction::CreateLambda([this]{ MediaPlayer->Previous(); }),
-		FCanExecuteAction::CreateLambda([this]{ return (MediaPlayer->GetPlaylistRef().Num() > 1); })
-	);
-
-	ToolkitCommands->MapAction(
-		Commands.ReverseMedia,
-		FExecuteAction::CreateLambda([this]{ MediaPlayer->SetRate(GetReverseRate()); } ),
-		FCanExecuteAction::CreateLambda([this]{ return MediaPlayer->IsReady() && MediaPlayer->SupportsRate(GetReverseRate(), false); })
-	);
-
-	ToolkitCommands->MapAction(
-		Commands.RewindMedia,
-		FExecuteAction::CreateLambda([this]{ MediaPlayer->Rewind(); }),
-		FCanExecuteAction::CreateLambda([this]{ return MediaPlayer->IsReady() && MediaPlayer->SupportsSeeking() && MediaPlayer->GetTime() > FTimespan::Zero(); })
-	);
-}
-
-
 BEGIN_SLATE_FUNCTION_BUILD_OPTIMIZATION
 void FMediaPlayerEditorToolkit::ExtendToolBar()
 {
@@ -344,7 +204,6 @@ void FMediaPlayerEditorToolkit::ExtendToolBar()
 		}
 	};
 
-
 	TSharedPtr<FExtender> ToolbarExtender = MakeShareable(new FExtender);
 
 	ToolbarExtender->AddToolBarExtension(
@@ -358,36 +217,6 @@ void FMediaPlayerEditorToolkit::ExtendToolBar()
 }
 END_SLATE_FUNCTION_BUILD_OPTIMIZATION
 
-
-float FMediaPlayerEditorToolkit::GetForwardRate() const
-{
-	float Rate = MediaPlayer->GetRate();
-
-	if (Rate < 1.0f)
-	{
-		Rate = 1.0f;
-	}
-
-	return 2.0f * Rate;
-}
-
-
-float FMediaPlayerEditorToolkit::GetReverseRate() const
-{
-	float Rate = MediaPlayer->GetRate();
-
-	if (Rate > -1.0f)
-	{
-		return -1.0f;
-	}
-
-	return 2.0f * Rate;
-}
-
-
-/* FMediaPlayerEditorToolkit callbacks
- *****************************************************************************/
-
 TSharedRef<SDockTab> FMediaPlayerEditorToolkit::HandleTabManagerSpawnTab(const FSpawnTabArgs& Args, FName TabIdentifier)
 {
 	TSharedPtr<SWidget> TabWidget = SNullWidget::NullWidget;
@@ -414,7 +243,8 @@ TSharedRef<SDockTab> FMediaPlayerEditorToolkit::HandleTabManagerSpawnTab(const F
 	}
 	else if (TabIdentifier == MediaPlayerEditorToolkit::ViewerTabId)
 	{
-		TabWidget = SNew(SMediaPlayerEditorViewer, *MediaPlayer, nullptr, Style, true);
+		TabWidget = SNew(SMediaPlayerEditorViewer, *MediaPlayer, nullptr, Style, true)
+			.Commands(ToolkitCommands);
 	}
 
 	return SNew(SDockTab)
@@ -424,5 +254,4 @@ TSharedRef<SDockTab> FMediaPlayerEditorToolkit::HandleTabManagerSpawnTab(const F
 		];
 }
 
-
 #undef LOCTEXT_NAMESPACE
diff --git a/Engine/Plugins/Media/MediaPlayerEditor/Source/MediaPlayerEditor/Private/Toolkits/MediaPlayerEditorToolkit.h b/Engine/Plugins/Media/MediaPlayerEditor/Source/MediaPlayerEditor/Private/Toolkits/MediaPlayerEditorToolkit.h
--- a/Engine/Plugins/Media/MediaPlayerEditor/Source/MediaPlayerEditor/Private/Toolkits/MediaPlayerEditorToolkit.h
+++ b/Engine/Plugins/Media/MediaPlayerEditor/Source/MediaPlayerEditor/Private/Toolkits/MediaPlayerEditorToolkit.h
@@ -2,25 +2,19 @@
 
 #pragma once
 
-#include "CoreMinimal.h"
-#include "UObject/GCObject.h"
-#include "EditorUndoClient.h"
-#include "Styling/ISlateStyle.h"
-#include "Framework/Docking/TabManager.h"
-#include "Toolkits/AssetEditorToolkit.h"
+#include "Toolkits/MediaPlayerEditorToolkitMediaPlayerBase.h"
 
+class FSpawnTabArgs;
+class SDockTab;
 class UMediaPlayer;
 
 /**
  * Implements an Editor toolkit for media players.
  */
 class FMediaPlayerEditorToolkit
-	: public FAssetEditorToolkit
-	, public FEditorUndoClient
-	, public FGCObject
+	: public FMediaPlayerEditorToolkitMediaPlayerBase
 {
 public:
-
 	/**
 	 * Creates and initializes a new instance.
 	 *
@@ -28,11 +22,6 @@ public:
 	 */
 	FMediaPlayerEditorToolkit(const TSharedRef<ISlateStyle>& InStyle);
 
-	/** Virtual destructor. */
-	virtual ~FMediaPlayerEditorToolkit();
-
-public:
-
 	/**
 	 * Initializes the editor tool kit.
 	 *
@@ -42,73 +31,32 @@ public:
 	 */
 	void Initialize(UMediaPlayer* InMediaPlayer, const EToolkitMode::Type InMode, const TSharedPtr<IToolkitHost>& InToolkitHost);
 
-public:
-
-	//~ FAssetEditorToolkit interface
-
-	virtual FString GetDocumentationLink() const override;
+	//~ Begin FAssetEditorToolkit
 	virtual void OnClose() override;
 	virtual void RegisterTabSpawners(const TSharedRef<FTabManager>& TabManager) override;
 	virtual void UnregisterTabSpawners(const TSharedRef<FTabManager>& TabManager) override;
+	//~ End FAssetEditorToolkit
 
-public:
-
-	//~ IToolkit interface
-
+	//~ Begin IToolkit
 	virtual FText GetBaseToolkitName() const override;
 	virtual FName GetToolkitFName() const override;
-	virtual FLinearColor GetWorldCentricTabColorScale() const override;
 	virtual FString GetWorldCentricTabPrefix() const override;
+	//~ End IToolkit
 
-public:
-
-	//~ FGCObject interface
-
-	virtual void AddReferencedObjects(FReferenceCollector& Collector) override;
+	//~ Begin FGCObject
 	virtual FString GetReferencerName() const override
 	{
 		return TEXT("FMediaPlayerEditorToolkit");
 	}
+	//~ End FGCObject
 
 protected:
-
-	//~ FEditorUndoClient interface
-
-	virtual void PostUndo(bool bSuccess) override;
-	virtual void PostRedo(bool bSuccess) override;
-
-protected:
-
-	/** Binds the UI commands to delegates. */
-	void BindCommands();
-
-	/** Builds the toolbar widget for the media player editor. */
-	void ExtendToolBar();
-
-	/**
-	 * Gets the playback rate for fast forward.
-	 *
-	 * @return Forward playback rate.
-	 */
-	float GetForwardRate() const;
-
-	/**
-	 * Gets the playback rate for reverse.
-	 *
-	 * @return Reverse playback rate.
-	 */
-	float GetReverseRate() const;
+	//~ Begin FMediaPlayerToolkitBase
+	virtual TSharedRef<FTabManager::FLayout> CreateLayout() override;
+	virtual void ExtendToolBar() override;
+	//~ End FMediaPlayerToolkitBase
 
 private:
-
 	/** Callback for spawning tabs. */
 	TSharedRef<SDockTab> HandleTabManagerSpawnTab(const FSpawnTabArgs& Args, FName TabIdentifier);
-
-private:
-
-	/** The media player asset being edited. */
-	TObjectPtr<UMediaPlayer> MediaPlayer;
-
-	/** Pointer to the style set to use for toolkits. */
-	TSharedRef<ISlateStyle> Style;
 };
diff --git a/Engine/Plugins/Media/MediaPlayerEditor/Source/MediaPlayerEditor/Private/Toolkits/MediaPlayerEditorToolkitBase.cpp b/Engine/Plugins/Media/MediaPlayerEditor/Source/MediaPlayerEditor/Private/Toolkits/MediaPlayerEditorToolkitBase.cpp
new file mode 100644
--- /dev/null
+++ b/Engine/Plugins/Media/MediaPlayerEditor/Source/MediaPlayerEditor/Private/Toolkits/MediaPlayerEditorToolkitBase.cpp
@@ -0,0 +1,64 @@
+// Copyright Epic Games, Inc. All Rights Reserved.
+
+#include "Toolkits/MediaPlayerEditorToolkitBase.h"
+
+#include "Editor.h"
+#include "EditorReimportHandler.h"
+
+FMediaPlayerEditorToolkitBase::FMediaPlayerEditorToolkitBase(const TSharedRef<ISlateStyle>& InStyle)
+	: Style(InStyle)
+{
+}
+
+FMediaPlayerEditorToolkitBase::~FMediaPlayerEditorToolkitBase()
+{
+	FReimportManager::Instance()->OnPreReimport().RemoveAll(this);
+	FReimportManager::Instance()->OnPostReimport().RemoveAll(this);
+
+	if (GEditor)
+	{
+		GEditor->UnregisterForUndo(this);
+	}
+}
+
+FString FMediaPlayerEditorToolkitBase::GetDocumentationLink() const
+{
+	return FString(TEXT("WorkingWithMedia/IntegratingMedia/MediaFramework"));
+}
+
+FLinearColor FMediaPlayerEditorToolkitBase::GetWorldCentricTabColorScale() const
+{
+	return FLinearColor(0.3f, 0.2f, 0.5f, 0.5f);
+}
+
+void FMediaPlayerEditorToolkitBase::PostUndo(bool bSuccess)
+{
+	// Do nothing
+}
+
+void FMediaPlayerEditorToolkitBase::PostRedo(bool bSuccess)
+{
+	PostUndo(bSuccess);
+}
+
+void FMediaPlayerEditorToolkitBase::Initialize(UObject* InAsset, FName InAppIdentifier, const EToolkitMode::Type InMode, const TSharedPtr<IToolkitHost>& InToolkitHost)
+{
+	GEditor->RegisterForUndo(this);
+
+	BindCommands();
+
+	const TSharedRef<FTabManager::FLayout> Layout = CreateLayout();
+
+	FAssetEditorToolkit::InitAssetEditor(
+		InMode,
+		InToolkitHost,
+		InAppIdentifier,
+		Layout,
+		/* bCreateDefaultStandaloneMenu */ true,
+		/* bCreateDefaultToolbar */ true,
+		InAsset
+	);
+
+	ExtendToolBar();
+	RegenerateMenusAndToolbars();
+}
diff --git a/Engine/Plugins/Media/MediaPlayerEditor/Source/MediaPlayerEditor/Private/Toolkits/MediaPlayerEditorToolkitBase.h b/Engine/Plugins/Media/MediaPlayerEditor/Source/MediaPlayerEditor/Private/Toolkits/MediaPlayerEditorToolkitBase.h
new file mode 100644
--- /dev/null
+++ b/Engine/Plugins/Media/MediaPlayerEditor/Source/MediaPlayerEditor/Private/Toolkits/MediaPlayerEditorToolkitBase.h
@@ -0,0 +1,65 @@
+// Copyright Epic Games, Inc. All Rights Reserved.
+
+#pragma once
+
+#include "EditorUndoClient.h"
+#include "Toolkits/AssetEditorToolkit.h"
+#include "UObject/GCObject.h"
+
+class ISlateStyle;
+
+/**
+ * Base class for all Media Player Editor Toolkits
+ */
+class FMediaPlayerEditorToolkitBase
+	: public FAssetEditorToolkit
+	, public FEditorUndoClient
+	, public FGCObject
+{
+public:
+	/**
+	 * Creates and initializes a new instance.
+	 *
+	 * @param InStyle The style set to use.
+	 */
+	FMediaPlayerEditorToolkitBase(const TSharedRef<ISlateStyle>& InStyle);
+
+	/** Virtual destructor. */
+	virtual ~FMediaPlayerEditorToolkitBase();
+
+	//~ Begin FAssetEditorToolkit
+	virtual FString GetDocumentationLink() const override;
+	//~ End FAssetEditorToolkit
+
+	//~ Begin IToolkit
+	virtual FLinearColor GetWorldCentricTabColorScale() const override;
+	//~ End IToolkit
+
+	//~ Begin FEditorUndoClient
+	virtual void PostUndo(bool bSuccess) override;
+	virtual void PostRedo(bool bSuccess) override;
+	//~ End FEditorUndoClient
+
+protected:
+	/** Pointer to the style set to use for toolkits. */
+	TSharedRef<ISlateStyle> Style;
+
+	/** Binds the UI commands to delegates. */
+	virtual void BindCommands() {}
+
+	/** Sets up the default layout for the toolkit */
+	virtual TSharedRef<FTabManager::FLayout> CreateLayout() = 0;
+
+	/** Builds the toolbar widget for the media player editor. */
+	virtual void ExtendToolBar() {}
+
+	/**
+	 * Initializes the editor tool kit.
+	 *
+	 * @param InAsset The object this toolkit is editing.
+	 * @param InAppIdentifier The identifier for this specific toolkit.
+	 * @param InMode The mode to create the toolkit in.
+	 * @param InToolkitHost The toolkit host.
+	 */
+	void Initialize(UObject* InAsset, FName InAppIdentifier, const EToolkitMode::Type InMode, const TSharedPtr<IToolkitHost>& InToolkitHost);
+};
diff --git a/Engine/Plugins/Media/MediaPlayerEditor/Source/MediaPlayerEditor/Private/Toolkits/MediaPlayerEditorToolkitMediaPlayerBase.cpp b/Engine/Plugins/Media/MediaPlayerEditor/Source/MediaPlayerEditor/Private/Toolkits/MediaPlayerEditorToolkitMediaPlayerBase.cpp
new file mode 100644
--- /dev/null
+++ b/Engine/Plugins/Media/MediaPlayerEditor/Source/MediaPlayerEditor/Private/Toolkits/MediaPlayerEditorToolkitMediaPlayerBase.cpp
@@ -0,0 +1,201 @@
+// Copyright Epic Games, Inc. All Rights Reserved.
+
+#include "Toolkits/MediaPlayerEditorToolkitMediaPlayerBase.h"
+
+#include "Framework/Commands/UICommandList.h"
+#include "MediaPlayer.h"
+#include "MediaPlaylist.h"
+#include "Models/MediaPlayerEditorCommands.h"
+#include "UObject/UObjectGlobals.h"
+
+FMediaPlayerEditorToolkitMediaPlayerBase::FMediaPlayerEditorToolkitMediaPlayerBase(const TSharedRef<ISlateStyle>& InStyle)
+	: FMediaPlayerEditorToolkitBase(InStyle)
+	, MediaPlayer(nullptr)
+{
+}
+
+void FMediaPlayerEditorToolkitMediaPlayerBase::OnClose()
+{
+	if (MediaPlayer != nullptr)
+	{
+		MediaPlayer->Close();
+	}
+}
+
+void FMediaPlayerEditorToolkitMediaPlayerBase::AddReferencedObjects(FReferenceCollector& InCollector)
+{
+	InCollector.AddReferencedObject(MediaPlayer);
+}
+
+float FMediaPlayerEditorToolkitMediaPlayerBase::GetForwardRate() const
+{
+	float Rate = MediaPlayer->GetRate();
+
+	if (Rate < 1.0f)
+	{
+		Rate = 1.0f;
+	}
+
+	return 2.0f * Rate;
+}
+
+float FMediaPlayerEditorToolkitMediaPlayerBase::GetReverseRate() const
+{
+	float Rate = MediaPlayer->GetRate();
+
+	if (Rate > -1.0f)
+	{
+		return -1.0f;
+	}
+
+	return 2.0f * Rate;
+}
+
+void FMediaPlayerEditorToolkitMediaPlayerBase::BindCommands()
+{
+	const FMediaPlayerEditorCommands& Commands = FMediaPlayerEditorCommands::Get();
+
+	ToolkitCommands->MapAction(
+		Commands.CloseMedia,
+		FExecuteAction::CreateLambda([this] { MediaPlayer->Close(); }),
+		FCanExecuteAction::CreateLambda([this] { return !MediaPlayer->GetUrl().IsEmpty(); })
+	);
+
+	ToolkitCommands->MapAction(
+		Commands.ForwardMedia,
+		FExecuteAction::CreateLambda([this]{ MediaPlayer->SetRate(GetForwardRate()); }),
+		FCanExecuteAction::CreateLambda([this]{ return MediaPlayer->IsReady() && MediaPlayer->SupportsRate(GetForwardRate(), false); })
+	);
+
+	ToolkitCommands->MapAction(
+		Commands.NextMedia,
+		FExecuteAction::CreateLambda([this]{ MediaPlayer->Next(); }),
+		FCanExecuteAction::CreateLambda([this]{ return (MediaPlayer->GetPlaylistRef().Num() > 1); })
+	);
+
+	ToolkitCommands->MapAction(
+		Commands.PauseMedia,
+		FExecuteAction::CreateLambda([this]{ MediaPlayer->Pause(); }),
+		FCanExecuteAction::CreateLambda([this]{ return MediaPlayer->CanPause() && !MediaPlayer->IsPaused(); })
+	);
+
+	ToolkitCommands->MapAction(
+		Commands.PlayMedia,
+		FExecuteAction::CreateLambda([this]{ MediaPlayer->Play(); }),
+		FCanExecuteAction::CreateLambda([this]{ return MediaPlayer->IsReady() && (!MediaPlayer->IsPlaying() || (MediaPlayer->GetRate() != 1.0f)); })
+	);
+
+	ToolkitCommands->MapAction(
+		Commands.TogglePlayPauseMedia,
+		FExecuteAction::CreateLambda(
+			[this]
+			{
+				if (MediaPlayer->IsPlaying())
+				{
+					MediaPlayer->Pause();
+				}
+				else
+				{
+					MediaPlayer->Play();
+				}
+			}),
+		FCanExecuteAction::CreateLambda([this] { return MediaPlayer->IsReady(); })
+	);
+
+	ToolkitCommands->MapAction(
+		Commands.PreviousMedia,
+		FExecuteAction::CreateLambda([this]{ MediaPlayer->Previous(); }),
+		FCanExecuteAction::CreateLambda([this]{ return (MediaPlayer->GetPlaylistRef().Num() > 1); })
+	);
+
+	ToolkitCommands->MapAction(
+		Commands.ReverseMedia,
+		FExecuteAction::CreateLambda([this]{ MediaPlayer->SetRate(GetReverseRate()); } ),
+		FCanExecuteAction::CreateLambda([this]{ return MediaPlayer->IsReady() && MediaPlayer->SupportsRate(GetReverseRate(), false); })
+	);
+
+	ToolkitCommands->MapAction(
+		Commands.RewindMedia,
+		FExecuteAction::CreateLambda([this]{ MediaPlayer->Rewind(); }),
+		FCanExecuteAction::CreateLambda([this]{ return MediaPlayer->IsReady() && MediaPlayer->SupportsSeeking() && MediaPlayer->GetTime() > FTimespan::Zero(); })
+	);
+
+	ToolkitCommands->MapAction(
+		Commands.FastForwardMedia,
+		FExecuteAction::CreateLambda([this] { MediaPlayer->Seek(MediaPlayer->GetDuration()); }),
+		FCanExecuteAction::CreateLambda([this] { return MediaPlayer->IsReady() && MediaPlayer->SupportsSeeking(); })
+	);
+
+	static const FTimespan TenSeconds = FTimespan::FromSeconds(10.f);
+
+	ToolkitCommands->MapAction(
+		Commands.StepBackward,
+		FExecuteAction::CreateLambda([this] 
+			{
+				// TODO: Convert to underlying media player command when framework support is added.
+				const float FrameRate = MediaPlayer->GetVideoTrackFrameRate(INDEX_NONE, INDEX_NONE);
+
+				if (!FMath::IsNearlyZero(FrameRate))
+				{
+					MediaPlayer->Seek(MediaPlayer->GetTime() - FTimespan::FromSeconds(1.f / FrameRate));
+				}
+			}
+		),
+		FCanExecuteAction::CreateLambda([this] { return MediaPlayer->IsReady() && MediaPlayer->IsPaused() && MediaPlayer->SupportsSeeking(); })
+	);
+
+	ToolkitCommands->MapAction(
+		Commands.JumpBackward,
+		FExecuteAction::CreateLambda(
+			[this]
+			{
+				const FTimespan CurrentTime = MediaPlayer->GetTime();
+
+				if (CurrentTime <= TenSeconds)
+				{
+					MediaPlayer->Seek(0);
+				}
+				else
+				{
+					MediaPlayer->Seek(MediaPlayer->GetTime() - FTimespan::FromSeconds(10.f));
+				}
+			}),
+		FCanExecuteAction::CreateLambda([this] { return MediaPlayer->IsReady() && MediaPlayer->SupportsSeeking(); })
+	);
+
+	ToolkitCommands->MapAction(
+		Commands.StepForward,
+		FExecuteAction::CreateLambda([this] 
+			{
+				// TODO: Convert to underlying media player command when framework support is added.
+				const float FrameRate = MediaPlayer->GetVideoTrackFrameRate(INDEX_NONE, INDEX_NONE);
+
+				if (!FMath::IsNearlyZero(FrameRate))
+				{
+					MediaPlayer->Seek(MediaPlayer->GetTime() + FTimespan::FromSeconds(1.f / FrameRate));
+				}
+			}
+		),
+		FCanExecuteAction::CreateLambda([this] { return MediaPlayer->IsReady() && MediaPlayer->IsPaused() && MediaPlayer->SupportsSeeking(); })
+	);
+
+	ToolkitCommands->MapAction(
+		Commands.JumpForward,
+		FExecuteAction::CreateLambda(
+			[this]
+			{
+				const FTimespan CurrentTime = MediaPlayer->GetTime();
+				const FTimespan Duration = MediaPlayer->GetDuration();
+
+				if ((Duration - CurrentTime) <= TenSeconds)
+				{
+					MediaPlayer->Seek(Duration);
+				}
+				else
+				{
+					MediaPlayer->Seek(MediaPlayer->GetTime() + FTimespan::FromSeconds(10.f));
+				}
+			}),
+		FCanExecuteAction::CreateLambda([this] { return MediaPlayer->IsReady() && MediaPlayer->SupportsSeeking(); })
+	);
+}
diff --git a/Engine/Plugins/Media/MediaPlayerEditor/Source/MediaPlayerEditor/Private/Toolkits/MediaPlayerEditorToolkitMediaPlayerBase.h b/Engine/Plugins/Media/MediaPlayerEditor/Source/MediaPlayerEditor/Private/Toolkits/MediaPlayerEditorToolkitMediaPlayerBase.h
new file mode 100644
--- /dev/null
+++ b/Engine/Plugins/Media/MediaPlayerEditor/Source/MediaPlayerEditor/Private/Toolkits/MediaPlayerEditorToolkitMediaPlayerBase.h
@@ -0,0 +1,56 @@
+// Copyright Epic Games, Inc. All Rights Reserved.
+
+#pragma once
+
+#include "Toolkits/MediaPlayerEditorToolkitBase.h"
+
+class UMediaPlayer;
+
+#define UE_API MEDIAPLAYEREDITOR_API
+
+/**
+ * Base class for all Media Player Editor Toolkits using Media Players
+ */
+class FMediaPlayerEditorToolkitMediaPlayerBase
+	: public FMediaPlayerEditorToolkitBase
+{
+public:
+	/**
+	 * Creates and initializes a new instance.
+	 *
+	 * @param InStyle The style set to use.
+	 */
+	FMediaPlayerEditorToolkitMediaPlayerBase(const TSharedRef<ISlateStyle>& InStyle);
+
+	//~ Begin FAssetEditorToolkit
+	virtual void OnClose() override;
+	//~ End FAssetEditorToolkit
+
+	//~ Begin FGCObject
+	virtual void AddReferencedObjects(FReferenceCollector& InCollector) override;
+	//~ End FGCObject
+
+protected:
+	/** The media player to play the media with. */
+	TObjectPtr<UMediaPlayer> MediaPlayer;
+
+	/**
+	 * Gets the playback rate for fast forward.
+	 *
+	 * @return Forward playback rate.
+	 */
+	float GetForwardRate() const;
+
+	/**
+	 * Gets the playback rate for reverse.
+	 *
+	 * @return Reverse playback rate.
+	 */
+	float GetReverseRate() const;
+
+	//~ Begin FMediaPlayerEditorToolkitBase
+	virtual void BindCommands() override;
+	//~ End FMediaPlayerEditorToolkitBase
+};
+
+#undef UE_API
diff --git a/Engine/Plugins/Media/MediaPlayerEditor/Source/MediaPlayerEditor/Private/Toolkits/MediaPlaylistEditorToolkit.cpp b/Engine/Plugins/Media/MediaPlayerEditor/Source/MediaPlayerEditor/Private/Toolkits/MediaPlaylistEditorToolkit.cpp
--- a/Engine/Plugins/Media/MediaPlayerEditor/Source/MediaPlayerEditor/Private/Toolkits/MediaPlaylistEditorToolkit.cpp
+++ b/Engine/Plugins/Media/MediaPlayerEditor/Source/MediaPlayerEditor/Private/Toolkits/MediaPlaylistEditorToolkit.cpp
@@ -1,18 +1,15 @@
 // Copyright Epic Games, Inc. All Rights Reserved.
 
 #include "Toolkits/MediaPlaylistEditorToolkit.h"
-#include "Editor.h"
-#include "Styling/AppStyle.h"
-#include "EditorReimportHandler.h"
+
+#include "Framework/Docking/TabManager.h"
 #include "MediaPlaylist.h"
 #include "Widgets/Docking/SDockTab.h"
 #include "Widgets/SMediaPlaylistEditorDetails.h"
 #include "Widgets/SMediaPlaylistEditorMedia.h"
 
-
 #define LOCTEXT_NAMESPACE "FMediaPlaylistEditorToolkit"
 
-
 namespace MediaPlaylistEditorToolkit
 {
 	static const FName AppIdentifier("MediaPlaylistEditorApp");
@@ -20,27 +17,10 @@ namespace MediaPlaylistEditorToolkit
 	static const FName MediaTabId("Media");
 }
 
-
-/* FMediaPlaylistEditorToolkit structors
- *****************************************************************************/
-
 FMediaPlaylistEditorToolkit::FMediaPlaylistEditorToolkit(const TSharedRef<ISlateStyle>& InStyle)
-	: MediaPlaylist(nullptr)
-	, Style(InStyle)
-{ }
-
-
-FMediaPlaylistEditorToolkit::~FMediaPlaylistEditorToolkit()
-{
-	FReimportManager::Instance()->OnPreReimport().RemoveAll(this);
-	FReimportManager::Instance()->OnPostReimport().RemoveAll(this);
-
-	GEditor->UnregisterForUndo(this);
-}
-
-
-/* FMediaPlaylistEditorToolkit interface
- *****************************************************************************/
+	: FMediaPlayerEditorToolkitBase(InStyle)
+	, MediaPlaylist(nullptr)
+{}
 
 void FMediaPlaylistEditorToolkit::Initialize(UMediaPlaylist* InMediaPlaylist, const EToolkitMode::Type InMode, const TSharedPtr<IToolkitHost>& InToolkitHost)
 {
@@ -53,10 +33,18 @@ void FMediaPlaylistEditorToolkit::Initialize(UMediaPlaylist* InMediaPlaylist, co
 
 	// Support undo/redo
 	MediaPlaylist->SetFlags(RF_Transactional);
-	GEditor->RegisterForUndo(this);
 
-	// create tab layout
-	const TSharedRef<FTabManager::FLayout> Layout = FTabManager::NewLayout("Standalone_MediaPlaylistEditor_v3")
+	FMediaPlayerEditorToolkitBase::Initialize(
+		InMediaPlaylist,
+		MediaPlaylistEditorToolkit::AppIdentifier,
+		InMode, 
+		InToolkitHost
+	);
+}
+
+TSharedRef<FTabManager::FLayout> FMediaPlaylistEditorToolkit::CreateLayout()
+{
+	return FTabManager::NewLayout("Standalone_MediaPlaylistEditor_v3")
 		->AddArea
 		(
 			FTabManager::NewPrimaryArea()
@@ -77,30 +65,8 @@ void FMediaPlaylistEditorToolkit::Initialize(UMediaPlaylist* InMediaPlaylist, co
 						->SetSizeCoefficient(0.45f)
 				)
 		);
-
-	FAssetEditorToolkit::InitAssetEditor(
-		InMode,
-		InToolkitHost,
-		MediaPlaylistEditorToolkit::AppIdentifier,
-		Layout,
-		true /*bCreateDefaultStandaloneMenu*/,
-		true /*bCreateDefaultToolbar*/,
-		InMediaPlaylist
-	);
-	
-	RegenerateMenusAndToolbars();
-}
-
-
-/* FAssetEditorToolkit interface
- *****************************************************************************/
-
-FString FMediaPlaylistEditorToolkit::GetDocumentationLink() const
-{
-	return FString(TEXT("WorkingWithMedia/IntegratingMedia/MediaFramework"));
 }
 
-
 void FMediaPlaylistEditorToolkit::RegisterTabSpawners(const TSharedRef<FTabManager>& InTabManager)
 {
 	WorkspaceMenuCategory = InTabManager->AddLocalWorkspaceMenuCategory(LOCTEXT("WorkspaceMenu_MediaPlaylistEditor", "Media Player Editor"));
@@ -119,7 +85,6 @@ void FMediaPlaylistEditorToolkit::RegisterTabSpawners(const TSharedRef<FTabManag
 		.SetIcon(FSlateIcon(Style->GetStyleSetName(), "MediaPlayerEditor.Tabs.Media"));
 }
 
-
 void FMediaPlaylistEditorToolkit::UnregisterTabSpawners(const TSharedRef<FTabManager>& InTabManager)
 {
 	FAssetEditorToolkit::UnregisterTabSpawners(InTabManager);
@@ -128,60 +93,26 @@ void FMediaPlaylistEditorToolkit::UnregisterTabSpawners(const TSharedRef<FTabMan
 	InTabManager->UnregisterTabSpawner(MediaPlaylistEditorToolkit::DetailsTabId);
 }
 
-
-/* IToolkit interface
- *****************************************************************************/
-
 FText FMediaPlaylistEditorToolkit::GetBaseToolkitName() const
 {
 	return LOCTEXT("AppLabel", "Media Playlist Editor");
 }
 
-
 FName FMediaPlaylistEditorToolkit::GetToolkitFName() const
 {
 	return FName("MediaPlaylistEditor");
 }
 
-
-FLinearColor FMediaPlaylistEditorToolkit::GetWorldCentricTabColorScale() const
-{
-	return FLinearColor(0.3f, 0.2f, 0.5f, 0.5f);
-}
-
-
 FString FMediaPlaylistEditorToolkit::GetWorldCentricTabPrefix() const
 {
 	return LOCTEXT("WorldCentricTabPrefix", "MediaPlaylist ").ToString();
 }
 
-
-/* FGCObject interface
- *****************************************************************************/
-
-void FMediaPlaylistEditorToolkit::AddReferencedObjects(FReferenceCollector& Collector)
-{
-	Collector.AddReferencedObject(MediaPlaylist);
-}
-
-
-/* FEditorUndoClient interface
-*****************************************************************************/
-
-void FMediaPlaylistEditorToolkit::PostUndo(bool bSuccess)
-{
-}
-
-
-void FMediaPlaylistEditorToolkit::PostRedo(bool bSuccess)
+void FMediaPlaylistEditorToolkit::AddReferencedObjects(FReferenceCollector& InCollector)
 {
-	PostUndo(bSuccess);
+	InCollector.AddReferencedObject(MediaPlaylist);
 }
 
-
-/* FMediaPlaylistEditorToolkit callbacks
- *****************************************************************************/
-
 TSharedRef<SDockTab> FMediaPlaylistEditorToolkit::HandleTabManagerSpawnTab(const FSpawnTabArgs& Args, FName TabIdentifier)
 {
 	TSharedPtr<SWidget> TabWidget = SNullWidget::NullWidget;
@@ -202,5 +133,4 @@ TSharedRef<SDockTab> FMediaPlaylistEditorToolkit::HandleTabManagerSpawnTab(const
 		];
 }
 
-
 #undef LOCTEXT_NAMESPACE
diff --git a/Engine/Plugins/Media/MediaPlayerEditor/Source/MediaPlayerEditor/Private/Toolkits/MediaPlaylistEditorToolkit.h b/Engine/Plugins/Media/MediaPlayerEditor/Source/MediaPlayerEditor/Private/Toolkits/MediaPlaylistEditorToolkit.h
--- a/Engine/Plugins/Media/MediaPlayerEditor/Source/MediaPlayerEditor/Private/Toolkits/MediaPlaylistEditorToolkit.h
+++ b/Engine/Plugins/Media/MediaPlayerEditor/Source/MediaPlayerEditor/Private/Toolkits/MediaPlaylistEditorToolkit.h
@@ -2,25 +2,19 @@
 
 #pragma once
 
-#include "CoreMinimal.h"
-#include "UObject/GCObject.h"
-#include "EditorUndoClient.h"
-#include "Styling/ISlateStyle.h"
-#include "Framework/Docking/TabManager.h"
-#include "Toolkits/AssetEditorToolkit.h"
+#include "Toolkits/MediaPlayerEditorToolkitBase.h"
 
+class FSpawnTabArgs;
+class SDockTab;
 class UMediaPlaylist;
 
 /**
  * Implements an Editor toolkit for media play lists.
  */
 class FMediaPlaylistEditorToolkit
-	: public FAssetEditorToolkit
-	, public FEditorUndoClient
-	, public FGCObject
+	: public FMediaPlayerEditorToolkitBase
 {
 public:
-
 	/**
 	 * Creates and initializes a new instance.
 	 *
@@ -28,11 +22,6 @@ public:
 	 */
 	FMediaPlaylistEditorToolkit(const TSharedRef<ISlateStyle>& InStyle);
 
-	/** Virtual destructor. */
-	virtual ~FMediaPlaylistEditorToolkit();
-
-public:
-
 	/**
 	 * Initializes the editor tool kit.
 	 *
@@ -42,50 +31,34 @@ public:
 	 */
 	void Initialize(UMediaPlaylist* InMediaPlaylist, const EToolkitMode::Type InMode, const TSharedPtr<IToolkitHost>& InToolkitHost);
 
-public:
-
-	//~ FAssetEditorToolkit interface
-
-	virtual FString GetDocumentationLink() const override;
+	//~ Begin FAssetEditorToolkit
 	virtual void RegisterTabSpawners(const TSharedRef<FTabManager>& TabManager) override;
 	virtual void UnregisterTabSpawners(const TSharedRef<FTabManager>& TabManager) override;
+	//~ End FAssetEditorToolkit
 
-public:
-
-	//~ IToolkit interface
-
+	//~ Begin IToolkit
 	virtual FText GetBaseToolkitName() const override;
 	virtual FName GetToolkitFName() const override;
-	virtual FLinearColor GetWorldCentricTabColorScale() const override;
 	virtual FString GetWorldCentricTabPrefix() const override;
+	//~ End IToolkit
 
-public:
-
-	//~ FGCObject interface
-
-	virtual void AddReferencedObjects(FReferenceCollector& Collector) override;
+	//~ Begin FGCObject
+	virtual void AddReferencedObjects(FReferenceCollector& InCollector) override;
 	virtual FString GetReferencerName() const override
 	{
 		return TEXT("FMediaPlaylistEditorToolkit");
 	}
+	//~ End FGCObject
 
 protected:
-
-	//~ FEditorUndoClient interface
-
-	virtual void PostUndo(bool bSuccess) override;
-	virtual void PostRedo(bool bSuccess) override;
+	//~ Begin FMediaPlayerToolkitBase
+	virtual TSharedRef<FTabManager::FLayout> CreateLayout() override;
+	//~ End FMediaPlayerToolkitBase
 
 private:
-
-	/** Callback for spawning tabs. */
-	TSharedRef<SDockTab> HandleTabManagerSpawnTab(const FSpawnTabArgs& Args, FName TabIdentifier);
-
-private:
-
 	/** The media play list asset being edited. */
 	TObjectPtr<UMediaPlaylist> MediaPlaylist;
 
-	/** Pointer to the style set to use for toolkits. */
-	TSharedRef<ISlateStyle> Style;
+	/** Callback for spawning tabs. */
+	TSharedRef<SDockTab> HandleTabManagerSpawnTab(const FSpawnTabArgs& Args, FName TabIdentifier);
 };
diff --git a/Engine/Plugins/Media/MediaPlayerEditor/Source/MediaPlayerEditor/Private/Toolkits/MediaSourceEditorToolkit.cpp b/Engine/Plugins/Media/MediaPlayerEditor/Source/MediaPlayerEditor/Private/Toolkits/MediaSourceEditorToolkit.cpp
--- a/Engine/Plugins/Media/MediaPlayerEditor/Source/MediaPlayerEditor/Private/Toolkits/MediaSourceEditorToolkit.cpp
+++ b/Engine/Plugins/Media/MediaPlayerEditor/Source/MediaPlayerEditor/Private/Toolkits/MediaSourceEditorToolkit.cpp
@@ -2,14 +2,14 @@
 
 #include "Toolkits/MediaSourceEditorToolkit.h"
 
-#include "Editor.h"
-#include "EditorReimportHandler.h"
+#include "CommonRenderResources.h"
 #include "Engine/TextureRenderTarget2D.h"
-#include "Framework/MultiBox/MultiBoxBuilder.h"
+#include "Framework/Docking/TabManager.h"
+#include "Framework/MultiBox/MultiBoxExtender.h"
+#include "GlobalShader.h"
 #include "IMediaModule.h"
 #include "IMediaPlayerFactory.h"
 #include "MediaPlayer.h"
-#include "MediaPlaylist.h"
 #include "MediaSource.h"
 #include "MediaTexture.h"
 #include "Models/MediaPlayerEditorCommands.h"
@@ -17,23 +17,24 @@
 #include "PixelShaderUtils.h"
 #include "PostProcess/DrawRectangle.h"
 #include "RenderGraphBuilder.h"
+#include "RenderGraphEvent.h"
 #include "RenderGraphUtils.h"
+#include "RHIResources.h"
 #include "ScreenPass.h"
+#include "Shader.h"
+#include "ShaderParameterStruct.h"
 #include "SlateOptMacros.h"
-#include "Styling/AppStyle.h"
 #include "TextureResource.h"
-#include "UObject/Package.h"
 #include "Widgets/Docking/SDockTab.h"
 #include "Widgets/SMediaPlayerEditorDetails.h"
 #include "Widgets/SMediaPlayerEditorMediaDetails.h"
-#include "Widgets/SMediaPlayerEditorPlaylist.h"
-#include "Widgets/SMediaPlayerEditorStats.h"
 #include "Widgets/SMediaPlayerEditorViewer.h"
 #include "Widgets/SMediaSourceEditorDetails.h"
+#include "UObject/UnrealType.h"
+#include "UObject/UObjectGlobals.h"
 
 #define LOCTEXT_NAMESPACE "FMediaSourceEditorToolkit"
 
-
 namespace MediaSourceEditorToolkit
 {
 	static const FName AppIdentifier("MediaSourceEditorApp");
@@ -43,29 +44,11 @@ namespace MediaSourceEditorToolkit
 	static const FName ViewerTabId("Viewer");
 }
 
-
-/* FMediaSourceEditorToolkit structors
- *****************************************************************************/
-
 FMediaSourceEditorToolkit::FMediaSourceEditorToolkit(const TSharedRef<ISlateStyle>& InStyle)
-	: MediaPlayer(nullptr)
+	: FMediaPlayerEditorToolkitMediaPlayerBase(InStyle)
 	, MediaSource(nullptr)
 	, MediaTexture(nullptr)
-	, Style(InStyle)
-{ }
-
-
-FMediaSourceEditorToolkit::~FMediaSourceEditorToolkit()
-{
-	FReimportManager::Instance()->OnPreReimport().RemoveAll(this);
-	FReimportManager::Instance()->OnPostReimport().RemoveAll(this);
-
-	GEditor->UnregisterForUndo(this);
-}
-
-
-/* FMediaSourceEditorToolkit interface
- *****************************************************************************/
+{}
 
 void FMediaSourceEditorToolkit::Initialize(UMediaSource* InMediaSource, const EToolkitMode::Type InMode, const TSharedPtr<IToolkitHost>& InToolkitHost)
 {
@@ -78,7 +61,6 @@ void FMediaSourceEditorToolkit::Initialize(UMediaSource* InMediaSource, const ET
 
 	// support undo/redo
 	MediaSource->SetFlags(RF_Transactional);
-	GEditor->RegisterForUndo(this);
 
 	MediaPlayer = NewObject<UMediaPlayer>(GetTransientPackage());
 	MediaPlayer->SetLooping(true);
@@ -93,10 +75,17 @@ void FMediaSourceEditorToolkit::Initialize(UMediaSource* InMediaSource, const ET
 		MediaTexture->UpdateResource();
 	}
 
-	BindCommands();
+	FMediaPlayerEditorToolkitBase::Initialize(
+		InMediaSource,
+		MediaSourceEditorToolkit::AppIdentifier,
+		InMode, 
+		InToolkitHost
+	);
+}
 
-	// create tab layout
-	const TSharedRef<FTabManager::FLayout> Layout = FTabManager::NewLayout("Standalone_MediaSourceEditor_v0.3")
+TSharedRef<FTabManager::FLayout> FMediaSourceEditorToolkit::CreateLayout()
+{
+	return FTabManager::NewLayout("Standalone_MediaSourceEditor_v0.3")
 		->AddArea
 		(
 			FTabManager::NewPrimaryArea()
@@ -133,40 +122,8 @@ void FMediaSourceEditorToolkit::Initialize(UMediaSource* InMediaSource, const ET
 						)
 				)
 		);
-
-	FAssetEditorToolkit::InitAssetEditor(
-		InMode,
-		InToolkitHost,
-		MediaSourceEditorToolkit::AppIdentifier,
-		Layout,
-		true /*bCreateDefaultStandaloneMenu*/,
-		true /*bCreateDefaultToolbar*/,
-		InMediaSource
-	);
-
-	ExtendToolBar();
-	RegenerateMenusAndToolbars();
 }
 
-
-/* FAssetEditorToolkit interface
- *****************************************************************************/
-
-FString FMediaSourceEditorToolkit::GetDocumentationLink() const
-{
-	return FString(TEXT("WorkingWithMedia/IntegratingMedia/MediaFramework"));
-}
-
-
-void FMediaSourceEditorToolkit::OnClose()
-{
-	if (MediaPlayer != nullptr)
-	{
-		MediaPlayer->Close();
-	}
-}
-
-
 void FMediaSourceEditorToolkit::RegisterTabSpawners(const TSharedRef<FTabManager>& InTabManager)
 {
 	WorkspaceMenuCategory = InTabManager->AddLocalWorkspaceMenuCategory(LOCTEXT("WorkspaceMenu_MediaSourceEditor", "Media Source Editor"));
@@ -195,7 +152,6 @@ void FMediaSourceEditorToolkit::RegisterTabSpawners(const TSharedRef<FTabManager
 		.SetIcon(FSlateIcon(Style->GetStyleSetName(), "MediaPlayerEditor.Tabs.Player"));
 }
 
-
 void FMediaSourceEditorToolkit::UnregisterTabSpawners(const TSharedRef<class FTabManager>& InTabManager)
 {
 	FAssetEditorToolkit::UnregisterTabSpawners(InTabManager);
@@ -206,63 +162,29 @@ void FMediaSourceEditorToolkit::UnregisterTabSpawners(const TSharedRef<class FTa
 	InTabManager->UnregisterTabSpawner(MediaSourceEditorToolkit::DetailsTabId);
 }
 
-
-/* IToolkit interface
- *****************************************************************************/
-
 FText FMediaSourceEditorToolkit::GetBaseToolkitName() const
 {
 	return LOCTEXT("AppLabel", "Media Source Editor");
 }
 
-
 FName FMediaSourceEditorToolkit::GetToolkitFName() const
 {
 	return FName("MediaSourceEditor");
 }
 
-
-FLinearColor FMediaSourceEditorToolkit::GetWorldCentricTabColorScale() const
-{
-	return FLinearColor(0.3f, 0.2f, 0.5f, 0.5f);
-}
-
-
 FString FMediaSourceEditorToolkit::GetWorldCentricTabPrefix() const
 {
 	return LOCTEXT("WorldCentricTabPrefix", "MediaSource ").ToString();
 }
 
-
-/* FGCObject interface
- *****************************************************************************/
-
-void FMediaSourceEditorToolkit::AddReferencedObjects(FReferenceCollector& Collector)
-{
-	Collector.AddReferencedObject(MediaPlayer);
-	Collector.AddReferencedObject(MediaSource);
-	Collector.AddReferencedObject(MediaTexture);
-}
-
-
-/* FEditorUndoClient interface
-*****************************************************************************/
-
-void FMediaSourceEditorToolkit::PostUndo(bool bSuccess)
+void FMediaSourceEditorToolkit::AddReferencedObjects(FReferenceCollector& InCollector)
 {
-	// do nothing
-}
-
+	FMediaPlayerEditorToolkitMediaPlayerBase::AddReferencedObjects(InCollector);
 
-void FMediaSourceEditorToolkit::PostRedo(bool bSuccess)
-{
-	PostUndo(bSuccess);
+	InCollector.AddReferencedObject(MediaSource);
+	InCollector.AddReferencedObject(MediaTexture);
 }
 
-
-/* FMediaSourceEditorToolkit implementation
- *****************************************************************************/
-
 void FMediaSourceEditorToolkit::ValidateDesiredPlayer(UMediaPlayer* InMediaPlayer, UMediaSource* InMediaSource)
 {
 	// Validate desired player
@@ -285,77 +207,30 @@ void FMediaSourceEditorToolkit::ValidateDesiredPlayer(UMediaPlayer* InMediaPlaye
 
 void FMediaSourceEditorToolkit::BindCommands()
 {
-	const FMediaPlayerEditorCommands& Commands = FMediaPlayerEditorCommands::Get();
-
-	ToolkitCommands->MapAction(
-		Commands.CloseMedia,
-		FExecuteAction::CreateLambda([this] { MediaPlayer->Close(); }),
-		FCanExecuteAction::CreateLambda([this] { return !MediaPlayer->GetUrl().IsEmpty(); })
-	);
-
-	ToolkitCommands->MapAction(
-		Commands.ForwardMedia,
-		FExecuteAction::CreateLambda([this]{ MediaPlayer->SetRate(GetForwardRate()); }),
-		FCanExecuteAction::CreateLambda([this]{ return MediaPlayer->IsReady() && MediaPlayer->SupportsRate(GetForwardRate(), false); })
-	);
+	FMediaPlayerEditorToolkitMediaPlayerBase::BindCommands();
 
-	ToolkitCommands->MapAction(
-		Commands.GenerateThumbnail,
-		FExecuteAction::CreateLambda([this] { GenerateThumbnail(); })
-	);
-
-	ToolkitCommands->MapAction(
-		Commands.NextMedia,
-		FExecuteAction::CreateLambda([this]{ MediaPlayer->Next(); }),
-		FCanExecuteAction::CreateLambda([this]{ return (MediaPlayer->GetPlaylistRef().Num() > 1); })
-	);
+	const FMediaPlayerEditorCommands& Commands = FMediaPlayerEditorCommands::Get();
 
 	ToolkitCommands->MapAction(
 		Commands.OpenMedia,
 		FExecuteAction::CreateLambda([this]
-		{
-			FMediaPlayerOptions Options;
-			Options.SetAllAsOptional();
-			Options.InternalCustomOptions.Emplace(MediaPlayerOptionValues::Environment(), MediaPlayerOptionValues::Environment_Preview());
-			Options.InternalCustomOptions.Emplace(MediaPlayerOptionValues::ParseTimecodeInfo(), FVariant());
-			ValidateDesiredPlayer(MediaPlayer, MediaSource);
-			MediaPlayer->OpenSourceWithOptions(MediaSource, Options);
-		}),
+			{
+				FMediaPlayerOptions Options;
+				Options.SetAllAsOptional();
+				Options.InternalCustomOptions.Emplace(MediaPlayerOptionValues::Environment(), MediaPlayerOptionValues::Environment_Preview());
+				Options.InternalCustomOptions.Emplace(MediaPlayerOptionValues::ParseTimecodeInfo(), FVariant());
+				ValidateDesiredPlayer(MediaPlayer, MediaSource);
+				MediaPlayer->OpenSourceWithOptions(MediaSource, Options);
+			}),
 		FCanExecuteAction::CreateLambda([this] { return true; })
 	);
 
 	ToolkitCommands->MapAction(
-		Commands.PauseMedia,
-		FExecuteAction::CreateLambda([this]{ MediaPlayer->Pause(); }),
-		FCanExecuteAction::CreateLambda([this]{ return MediaPlayer->CanPause() && !MediaPlayer->IsPaused(); })
-	);
-
-	ToolkitCommands->MapAction(
-		Commands.PlayMedia,
-		FExecuteAction::CreateLambda([this]{ MediaPlayer->Play(); }),
-		FCanExecuteAction::CreateLambda([this]{ return MediaPlayer->IsReady() && (!MediaPlayer->IsPlaying() || (MediaPlayer->GetRate() != 1.0f)); })
-	);
-
-	ToolkitCommands->MapAction(
-		Commands.PreviousMedia,
-		FExecuteAction::CreateLambda([this]{ MediaPlayer->Previous(); }),
-		FCanExecuteAction::CreateLambda([this]{ return (MediaPlayer->GetPlaylistRef().Num() > 1); })
-	);
-
-	ToolkitCommands->MapAction(
-		Commands.ReverseMedia,
-		FExecuteAction::CreateLambda([this]{ MediaPlayer->SetRate(GetReverseRate()); } ),
-		FCanExecuteAction::CreateLambda([this]{ return MediaPlayer->IsReady() && MediaPlayer->SupportsRate(GetReverseRate(), false); })
-	);
-
-	ToolkitCommands->MapAction(
-		Commands.RewindMedia,
-		FExecuteAction::CreateLambda([this]{ MediaPlayer->Rewind(); }),
-		FCanExecuteAction::CreateLambda([this]{ return MediaPlayer->IsReady() && MediaPlayer->SupportsSeeking() && MediaPlayer->GetTime() > FTimespan::Zero(); })
+		Commands.GenerateThumbnail,
+		FExecuteAction::CreateLambda([this] { GenerateThumbnail(); })
 	);
 }
 
-
 BEGIN_SLATE_FUNCTION_BUILD_OPTIMIZATION
 void FMediaSourceEditorToolkit::ExtendToolBar()
 {
@@ -399,32 +274,6 @@ void FMediaSourceEditorToolkit::ExtendToolBar()
 }
 END_SLATE_FUNCTION_BUILD_OPTIMIZATION
 
-
-float FMediaSourceEditorToolkit::GetForwardRate() const
-{
-	float Rate = MediaPlayer->GetRate();
-
-	if (Rate < 1.0f)
-	{
-		Rate = 1.0f;
-	}
-
-	return 2.0f * Rate;
-}
-
-
-float FMediaSourceEditorToolkit::GetReverseRate() const
-{
-	float Rate = MediaPlayer->GetRate();
-
-	if (Rate > -1.0f)
-	{
-		return -1.0f;
-	}
-
-	return 2.0f * Rate;
-}
-
 void FMediaSourceEditorToolkit::GenerateThumbnail()
 {
 	// Create render target.
@@ -503,10 +352,6 @@ void FMediaSourceEditorToolkit::GenerateThumbnail()
 	}
 }
 
-
-/* FMediaSourceEditorToolkit callbacks
- *****************************************************************************/
-
 TSharedRef<SDockTab> FMediaSourceEditorToolkit::HandleTabManagerSpawnTab(const FSpawnTabArgs& Args, FName TabIdentifier)
 {
 	TSharedPtr<SWidget> TabWidget = SNullWidget::NullWidget;
@@ -526,7 +371,8 @@ TSharedRef<SDockTab> FMediaSourceEditorToolkit::HandleTabManagerSpawnTab(const F
 	else if (TabIdentifier == MediaSourceEditorToolkit::ViewerTabId)
 	{
 		TabWidget = SNew(SMediaPlayerEditorViewer, *MediaPlayer, MediaTexture, Style, true)
-			.bShowUrl(false);
+			.bShowUrl(false)
+			.Commands(ToolkitCommands);
 	}
 
 	return SNew(SDockTab)
@@ -536,5 +382,4 @@ TSharedRef<SDockTab> FMediaSourceEditorToolkit::HandleTabManagerSpawnTab(const F
 		];
 }
 
-
 #undef LOCTEXT_NAMESPACE
diff --git a/Engine/Plugins/Media/MediaPlayerEditor/Source/MediaPlayerEditor/Private/Toolkits/MediaSourceEditorToolkit.h b/Engine/Plugins/Media/MediaPlayerEditor/Source/MediaPlayerEditor/Private/Toolkits/MediaSourceEditorToolkit.h
--- a/Engine/Plugins/Media/MediaPlayerEditor/Source/MediaPlayerEditor/Private/Toolkits/MediaSourceEditorToolkit.h
+++ b/Engine/Plugins/Media/MediaPlayerEditor/Source/MediaPlayerEditor/Private/Toolkits/MediaSourceEditorToolkit.h
@@ -2,14 +2,10 @@
 
 #pragma once
 
-#include "EditorUndoClient.h"
-#include "Framework/Docking/TabManager.h"
-#include "Styling/ISlateStyle.h"
-#include "Toolkits/AssetEditorToolkit.h"
-#include "UObject/GCObject.h"
-
-#define UE_API MEDIAPLAYEREDITOR_API
+#include "Toolkits/MediaPlayerEditorToolkitMediaPlayerBase.h"
 
+class FSpawnTabArgs;
+class SDockTab;
 class UMediaPlayer;
 class UMediaSource;
 class UMediaTexture;
@@ -18,23 +14,15 @@ class UMediaTexture;
  * Implements an Editor toolkit for media sources.
  */
 class FMediaSourceEditorToolkit
-	: public FAssetEditorToolkit
-	, public FEditorUndoClient
-	, public FGCObject
+	: public FMediaPlayerEditorToolkitMediaPlayerBase
 {
 public:
-
 	/**
 	 * Creates and initializes a new instance.
 	 *
 	 * @param InStyle The style set to use.
 	 */
-	UE_API FMediaSourceEditorToolkit(const TSharedRef<ISlateStyle>& InStyle);
-
-	/** Virtual destructor. */
-	UE_API virtual ~FMediaSourceEditorToolkit();
-
-public:
+	FMediaSourceEditorToolkit(const TSharedRef<ISlateStyle>& InStyle);
 
 	/**
 	 * Initializes the editor tool kit.
@@ -43,84 +31,40 @@ public:
 	 * @param InMode The mode to create the toolkit in.
 	 * @param InToolkitHost The toolkit host.
 	 */
-	UE_API void Initialize(UMediaSource* InMediaSource, const EToolkitMode::Type InMode, const TSharedPtr<IToolkitHost>& InToolkitHost);
-
-public:
-
-	//~ FAssetEditorToolkit interface
-
-	UE_API virtual FString GetDocumentationLink() const override;
-	UE_API virtual void OnClose() override;
-	UE_API virtual void RegisterTabSpawners(const TSharedRef<FTabManager>& TabManager) override;
-	UE_API virtual void UnregisterTabSpawners(const TSharedRef<FTabManager>& TabManager) override;
-
-public:
+	void Initialize(UMediaSource* InMediaSource, const EToolkitMode::Type InMode, const TSharedPtr<IToolkitHost>& InToolkitHost);
 
-	//~ IToolkit interface
+	//~ Begin FAssetEditorToolkit
+	virtual void RegisterTabSpawners(const TSharedRef<FTabManager>& TabManager) override;
+	virtual void UnregisterTabSpawners(const TSharedRef<FTabManager>& TabManager) override;
+	//~ End FAssetEditorToolkit
 
-	UE_API virtual FText GetBaseToolkitName() const override;
-	UE_API virtual FName GetToolkitFName() const override;
-	UE_API virtual FLinearColor GetWorldCentricTabColorScale() const override;
-	UE_API virtual FString GetWorldCentricTabPrefix() const override;
+	//~ Begin IToolkit
+	virtual FText GetBaseToolkitName() const override;
+	virtual FName GetToolkitFName() const override;
+	virtual FString GetWorldCentricTabPrefix() const override;
+	//~ End IToolkit
 
-public:
-
-	//~ FGCObject interface
-
-	UE_API virtual void AddReferencedObjects(FReferenceCollector& Collector) override;
+	//~ Begin FGCObject
+	virtual void AddReferencedObjects(FReferenceCollector& InCollector) override;
 	virtual FString GetReferencerName() const override
 	{
 		return TEXT("FMediaSourceEditorToolkit");
 	}
+	//~ End FGCObject
 
 protected:
+	//~ Begin FMediaPlayerToolkitBase
+	virtual TSharedRef<FTabManager::FLayout> CreateLayout() override;
+	virtual void BindCommands() override;
+	virtual void ExtendToolBar() override;
+	//~ End FMediaPlayerToolkitBase
 
-	//~ FEditorUndoClient interface
-
-	UE_API virtual void PostUndo(bool bSuccess) override;
-	UE_API virtual void PostRedo(bool bSuccess) override;
-
-protected:
+private:
 	/**
 	 * Validates that the given player can open the media source with the selected desired player.
 	 * If not, the desired player is reset to fallback to the media source selection.
 	 */
-	static UE_API void ValidateDesiredPlayer(UMediaPlayer* InMediaPlayer, UMediaSource* InMediaSource);
-
-	/** Binds the UI commands to delegates. */
-	UE_API void BindCommands();
-
-	/** Builds the toolbar widget for the media player editor. */
-	UE_API void ExtendToolBar();
-
-	/**
-	 * Gets the playback rate for fast forward.
-	 *
-	 * @return Forward playback rate.
-	 */
-	UE_API float GetForwardRate() const;
-
-	/**
-	 * Gets the playback rate for reverse.
-	 *
-	 * @return Reverse playback rate.
-	 */
-	UE_API float GetReverseRate() const;
-
-	/**
-	 * Enqueues rendering commands to generate a thumbnail.
-	 */
-	UE_API void GenerateThumbnail();
-
-private:
-
-	/** Callback for spawning tabs. */
-	UE_API TSharedRef<SDockTab> HandleTabManagerSpawnTab(const FSpawnTabArgs& Args, FName TabIdentifier);
-
-private:
-
-	/** The media player to play the media with. */
-	TObjectPtr<UMediaPlayer> MediaPlayer;
+	static void ValidateDesiredPlayer(UMediaPlayer* InMediaPlayer, UMediaSource* InMediaSource);
 
 	/** The media source asset being edited. */
 	TObjectPtr<UMediaSource> MediaSource;
@@ -128,8 +72,9 @@ private:
 	/** The media texture to output the media to. */
 	TObjectPtr<UMediaTexture> MediaTexture;
 
-	/** Pointer to the style set to use for toolkits. */
-	TSharedRef<ISlateStyle> Style;
-};
+	/** Callback for spawning tabs. */
+	TSharedRef<SDockTab> HandleTabManagerSpawnTab(const FSpawnTabArgs& Args, FName TabIdentifier);
 
-#undef UE_API
+	/** Enqueues rendering commands to generate a thumbnail. */
+	void GenerateThumbnail();
+};
diff --git a/Engine/Plugins/Media/MediaPlayerEditor/Source/MediaPlayerEditor/Private/Widgets/SMediaPlayerEditorViewer.cpp b/Engine/Plugins/Media/MediaPlayerEditor/Source/MediaPlayerEditor/Private/Widgets/SMediaPlayerEditorViewer.cpp
--- a/Engine/Plugins/Media/MediaPlayerEditor/Source/MediaPlayerEditor/Private/Widgets/SMediaPlayerEditorViewer.cpp
+++ b/Engine/Plugins/Media/MediaPlayerEditor/Source/MediaPlayerEditor/Private/Widgets/SMediaPlayerEditorViewer.cpp
@@ -350,7 +350,7 @@ void SMediaPlayerEditorViewer::Construct(const FArguments& InArgs, UMediaPlayer&
 							.Padding(0.0f)
 							[
 								// movie area
-								SAssignNew(PlayerViewport, SMediaPlayerEditorViewport, InMediaPlayer, InMediaTexture, InStyle, bInIsSoundEnabled)
+								SAssignNew(PlayerViewport, SMediaPlayerEditorViewport, InMediaPlayer, InMediaTexture, InStyle, bInIsSoundEnabled, InArgs._Commands)
 							]
 					]
 
diff --git a/Engine/Plugins/Media/MediaPlayerEditor/Source/MediaPlayerEditor/Private/Widgets/SMediaPlayerEditorViewport.cpp b/Engine/Plugins/Media/MediaPlayerEditor/Source/MediaPlayerEditor/Private/Widgets/SMediaPlayerEditorViewport.cpp
--- a/Engine/Plugins/Media/MediaPlayerEditor/Source/MediaPlayerEditor/Private/Widgets/SMediaPlayerEditorViewport.cpp
+++ b/Engine/Plugins/Media/MediaPlayerEditor/Source/MediaPlayerEditor/Private/Widgets/SMediaPlayerEditorViewport.cpp
@@ -2,8 +2,9 @@
 
 #include "Widgets/SMediaPlayerEditorViewport.h"
 
+#include "Framework/Commands/UICommandList.h"
 #include "MediaPlayer.h"
-#include "Misc/Paths.h"
+#include "Models/MediaPlayerEditorCommands.h"
 #include "Widgets/SBoxPanel.h"
 #include "Widgets/SOverlay.h"
 #include "Widgets/Layout/SBorder.h"
@@ -32,10 +33,11 @@ SMediaPlayerEditorViewport::SMediaPlayerEditorViewport()
 
 void SMediaPlayerEditorViewport::Construct(const FArguments& InArgs, UMediaPlayer& InMediaPlayer,
 	UMediaTexture* InMediaTexture, const TSharedRef<ISlateStyle>& InStyle,
-	bool bInIsSoundEnabled)
+	bool bInIsSoundEnabled, TSharedPtr<FUICommandList> InCommands)
 {
 	MediaPlayer = &InMediaPlayer;
 	Style = InStyle;
+	CommandsWeak = InCommands;
 
 	const auto OverlayFont = Style->GetFontStyle("MediaPlayerEditor.ViewportFont");
 
@@ -157,79 +159,28 @@ void SMediaPlayerEditorViewport::Construct(const FArguments& InArgs, UMediaPlaye
 /* SWidget interface
  *****************************************************************************/
 
-FReply SMediaPlayerEditorViewport::OnMouseButtonDown(const FGeometry& MyGeometry, const FPointerEvent& MouseEvent)
-{
-	if ((bIsMouseControlEnabled == false) ||
-		(MouseEvent.GetEffectingButton() != EKeys::LeftMouseButton))
-	{
-		return FReply::Unhandled();
-	}
-
-	SetCursor(EMouseCursor::CardinalCross);
-
-	return FReply::Handled().CaptureMouse(SharedThis(this));
-}
-
-
 FReply SMediaPlayerEditorViewport::OnMouseButtonDoubleClick(const FGeometry& MyGeometry, const FPointerEvent& MouseEvent)
 {
-	if ((bIsMouseControlEnabled == false) ||
-		(MouseEvent.GetEffectingButton() != EKeys::LeftMouseButton))
-	{
-		return FReply::Unhandled();
-	}
-
-	if (MediaPlayer->IsPlaying())
-	{
-		MediaPlayer->Pause();
-	}
-	else
+	if (TSharedPtr<FUICommandList> Commands = CommandsWeak.Pin())
 	{
-		MediaPlayer->Play();
+		Commands->ExecuteAction(FMediaPlayerEditorCommands::Get().TogglePlayPauseMedia.ToSharedRef());
 	}
 
 	return FReply::Handled();
 }
 
 
-FReply SMediaPlayerEditorViewport::OnMouseButtonUp(const FGeometry& MyGeometry, const FPointerEvent& MouseEvent)
+FReply SMediaPlayerEditorViewport::OnKeyDown(const FGeometry& MyGeometry, const FKeyEvent& InKeyEvent)
 {
-	if ((bIsMouseControlEnabled == false) ||
-		(MouseEvent.GetEffectingButton() != EKeys::LeftMouseButton))
+	if (TSharedPtr<FUICommandList> Commands = CommandsWeak.Pin())
 	{
-		return FReply::Unhandled();
-	}
-
-	SetCursor(EMouseCursor::Default);
-
-	return FReply::Handled().ReleaseMouseCapture();
-}
-
-
-FReply SMediaPlayerEditorViewport::OnMouseMove(const FGeometry& MyGeometry, const FPointerEvent& MouseEvent)
-{
-	if (!HasMouseCapture())
-	{
-		return FReply::Unhandled();
-	}
-
-	const FVector2D CursorDelta = MouseEvent.GetCursorDelta();
-	MediaPlayer->SetViewRotation(FRotator::MakeFromEuler(FVector(CursorDelta.Y, CursorDelta.X, 0.0f)), false);
-
-	return FReply::Handled();
-}
-
-
-FReply SMediaPlayerEditorViewport::OnMouseWheel(const FGeometry& MyGeometry, const FPointerEvent& MouseEvent)
-{
-	if (bIsMouseControlEnabled == false)
-	{
-		return FReply::Unhandled();
+		if (Commands->ProcessCommandBindings(InKeyEvent))
+		{
+			return FReply::Handled();
+		}
 	}
 
-	MediaPlayer->SetViewField(MouseEvent.GetWheelDelta(), 0.0f, false);
-
-	return FReply::Handled();
+	return FReply::Unhandled();
 }
 
 
diff --git a/Engine/Plugins/Media/MediaPlayerEditor/Source/MediaPlayerEditor/Private/Widgets/SMediaPlayerEditorViewport.h b/Engine/Plugins/Media/MediaPlayerEditor/Source/MediaPlayerEditor/Private/Widgets/SMediaPlayerEditorViewport.h
--- a/Engine/Plugins/Media/MediaPlayerEditor/Source/MediaPlayerEditor/Private/Widgets/SMediaPlayerEditorViewport.h
+++ b/Engine/Plugins/Media/MediaPlayerEditor/Source/MediaPlayerEditor/Private/Widgets/SMediaPlayerEditorViewport.h
@@ -6,6 +6,7 @@
 #include "Widgets/DeclarativeSyntaxSupport.h"
 #include "Widgets/SCompoundWidget.h"
 
+class FUICommandList;
 class UMediaPlayer;
 class UMediaSoundComponent;
 class UMediaTexture;
@@ -32,12 +33,13 @@ public:
 	 * @param InArgs The declaration data for this widget.
 	 * @param InMediaPlayer The UMediaPlayer asset to show the details for.
 	 * @param InMediaTexture The UMediaTexture asset to output video to. If nullptr then use our own.
-	 * @param InStyleSet The style set to use.
+	 * @param InStyle The style set to use.
 	 * @param bInIsSoundEnabled If true then produce sound.
+	 * @param InCommands Command list for keybinds.
 	 */
 	void Construct(const FArguments& InArgs, UMediaPlayer& InMediaPlayer,
 		UMediaTexture* InMediaTexture, const TSharedRef<ISlateStyle>& InStyle,
-		bool bInIsSoundEnabled);
+		bool bInIsSoundEnabled, TSharedPtr<FUICommandList> InCommands);
 
 	/**
 	 * Enables/disables using the mouse to control the viewport.
@@ -48,11 +50,8 @@ public:
 
 	//~ SWidget interface
 
-	virtual FReply OnMouseButtonDown(const FGeometry& MyGeometry, const FPointerEvent& MouseEvent) override;
 	virtual FReply OnMouseButtonDoubleClick(const FGeometry& MyGeometry, const FPointerEvent& MouseEvent) override;
-	virtual FReply OnMouseButtonUp(const FGeometry& MyGeometry, const FPointerEvent& MouseEvent) override;
-	virtual FReply OnMouseMove(const FGeometry& MyGeometry, const FPointerEvent& MouseEvent) override;
-	virtual FReply OnMouseWheel(const FGeometry& MyGeometry, const FPointerEvent& MouseEvent) override;
+	virtual FReply OnKeyDown(const FGeometry& MyGeometry, const FKeyEvent& InKeyEvent) override;
 
 private:
 
@@ -81,4 +80,7 @@ private:
 	
 	/** True if the mouse can control things. */
 	bool bIsMouseControlEnabled;
+
+	/** Command list for executing commands. */
+	TWeakPtr<FUICommandList> CommandsWeak;
 };
\ No newline at end of file
diff --git a/Engine/Plugins/Media/MediaPlayerEditor/Source/MediaPlayerEditor/Public/Models/MediaPlayerEditorCommands.h b/Engine/Plugins/Media/MediaPlayerEditor/Source/MediaPlayerEditor/Public/Models/MediaPlayerEditorCommands.h
new file mode 100644
--- /dev/null
+++ b/Engine/Plugins/Media/MediaPlayerEditor/Source/MediaPlayerEditor/Public/Models/MediaPlayerEditorCommands.h
@@ -0,0 +1,81 @@
+// Copyright Epic Games, Inc. All Rights Reserved.
+
+#pragma once
+
+#include "CoreMinimal.h"
+#include "Framework/Commands/Commands.h"
+
+#define UE_API MEDIAPLAYEREDITOR_API
+
+/**
+ * Holds the UI commands for the MediaPlayerEditorToolkit widget.
+ */
+class FMediaPlayerEditorCommands
+	: public TCommands<FMediaPlayerEditorCommands>
+{
+public:
+	UE_API static const FMediaPlayerEditorCommands& GetExternal();
+
+	/** Default constructor. */
+	FMediaPlayerEditorCommands() 
+		: TCommands<FMediaPlayerEditorCommands>("MediaPlayerEditor", NSLOCTEXT("Contexts", "MediaPlayerEditor", "MediaPlayer Editor"), NAME_None, "MediaPlayerEditorStyle")
+	{ }
+
+public:
+
+	//~ TCommands interface
+
+	virtual void RegisterCommands() override;
+	
+public:
+
+	/** Close the currently opened media. */
+	TSharedPtr<FUICommandInfo> CloseMedia;
+
+	/** Plays media forward. */
+	TSharedPtr<FUICommandInfo> ForwardMedia;
+
+	/** Generate a thumnbnail. */
+	TSharedPtr<FUICommandInfo> GenerateThumbnail;
+
+	/** Jump to next item in the play list. */
+	TSharedPtr<FUICommandInfo> NextMedia;
+
+	/** Open the current media. */
+	TSharedPtr<FUICommandInfo> OpenMedia;
+
+	/** Pauses media playback. */
+	TSharedPtr<FUICommandInfo> PauseMedia;
+
+	/** Starts media playback. */
+	TSharedPtr<FUICommandInfo> PlayMedia;
+
+	/** Toggles between media play and pause. */
+	TSharedPtr<FUICommandInfo> TogglePlayPauseMedia;
+
+	/** Jump to previous item in the play list. */
+	TSharedPtr<FUICommandInfo> PreviousMedia;
+
+	/** Reverses media playback. */
+	TSharedPtr<FUICommandInfo> ReverseMedia;
+
+	/** Rewinds the media to the beginning. */
+	TSharedPtr<FUICommandInfo> RewindMedia;
+
+	/** Fast forwards the media to the end. */
+	TSharedPtr<FUICommandInfo> FastForwardMedia;
+
+	/** Step forward media 1 frame */
+	TSharedPtr<FUICommandInfo> StepForward;
+
+	/** Step backward media 1 frame. */
+	TSharedPtr<FUICommandInfo> StepBackward;
+
+	/** Jump forward media 10s */
+	TSharedPtr<FUICommandInfo> JumpForward;
+
+	/** Jump backward media 10s. */
+	TSharedPtr<FUICommandInfo> JumpBackward;
+};
+
+#undef UE_API
diff --git a/Engine/Plugins/Media/MediaPlayerEditor/Source/MediaPlayerEditor/Public/Widgets/SMediaPlayerEditorViewer.h b/Engine/Plugins/Media/MediaPlayerEditor/Source/MediaPlayerEditor/Public/Widgets/SMediaPlayerEditorViewer.h
--- a/Engine/Plugins/Media/MediaPlayerEditor/Source/MediaPlayerEditor/Public/Widgets/SMediaPlayerEditorViewer.h
+++ b/Engine/Plugins/Media/MediaPlayerEditor/Source/MediaPlayerEditor/Public/Widgets/SMediaPlayerEditorViewer.h
@@ -14,6 +14,7 @@
 #define UE_API MEDIAPLAYEREDITOR_API
 
 class FMenuBuilder;
+class FUICommandList;
 class SEditableTextBox;
 class SMediaPlayerEditorViewport;
 class UMediaPlayer;
@@ -38,6 +39,8 @@ public:
 
 		/** If the url should be shown. */
 		SLATE_ARGUMENT(bool, bShowUrl)
+		/** Command list for executing keybinds and mouse clicks. */
+		SLATE_ARGUMENT(TSharedPtr<FUICommandList>, Commands)
 
 	SLATE_END_ARGS()
 
