Media Player Editor: You can now filter by texture color channel.

#Synced-CL 50961346

diff --git a/Engine/Plugins/Media/MediaPlate/Source/MediaPlateEditor/Private/Toolkits/MediaPlateEditorToolkit.cpp b/Engine/Plugins/Media/MediaPlate/Source/MediaPlateEditor/Private/Toolkits/MediaPlateEditorToolkit.cpp
--- a/Engine/Plugins/Media/MediaPlate/Source/MediaPlateEditor/Private/Toolkits/MediaPlateEditorToolkit.cpp
+++ b/Engine/Plugins/Media/MediaPlate/Source/MediaPlateEditor/Private/Toolkits/MediaPlateEditorToolkit.cpp
@@ -17,6 +17,8 @@
 #include "SlateOptMacros.h"
 #include "Styling/AppStyle.h"
 #include "Widgets/Docking/SDockTab.h"
+#include "Widgets/Input/SCheckBox.h"
+#include "Widgets/SMediaImageTextureChannelToggle.h"
 #include "Widgets/SMediaPlateEditorDetails.h"
 #include "Widgets/SMediaPlateEditorMediaDetails.h"
 #include "Widgets/SMediaPlateEditorPlaylist.h"
@@ -158,6 +160,27 @@ void FMediaPlateEditorToolkit::Initialize(UMediaPlateComponent* InMediaPlate, co
 /* FAssetEditorToolkit interface
  *****************************************************************************/
 
+void FMediaPlateEditorToolkit::ToggleChannelMask(MediaPlayerEditor::MediaImage::ETextureChannelMask InChannelToToggle)
+{
+	if (!Viewer.IsValid())
+	{
+		return;
+	}
+
+	using namespace MediaPlayerEditor::MediaImage;
+
+	const ETextureChannelMask CurrentMask = Viewer->GetChannelMask();
+
+	if (EnumHasAllFlags(CurrentMask, InChannelToToggle))
+	{
+		Viewer->SetChannelMask(CurrentMask & ~InChannelToToggle);
+	}
+	else
+	{
+		Viewer->SetChannelMask(CurrentMask | InChannelToToggle);
+	}
+}
+
 FString FMediaPlateEditorToolkit::GetDocumentationLink() const
 {
 	return FString(TEXT("WorkingWithMedia/IntegratingMedia/MediaFramework"));
@@ -350,45 +373,81 @@ void FMediaPlateEditorToolkit::BindCommands()
 		FExecuteAction::CreateSP(this, &FMediaPlateEditorToolkit::HandleMediaPlateEvent, EMediaPlateEventState::JumpBackward),
 		FCanExecuteAction::CreateSP(this, &FMediaPlateEditorToolkit::CanHandleMediaPlateEvent, EMediaPlateEventState::JumpBackward)
 	);
+
+	ToolkitCommands->MapAction(
+		Commands.ToggledRedTextureChannel,
+		FExecuteAction::CreateLambda([this] { ToggleChannelMask(MediaPlayerEditor::MediaImage::ETextureChannelMask::Red); }),
+		FCanExecuteAction::CreateLambda([this] { return Viewer.IsValid(); }),
+		FIsActionChecked::CreateLambda([this] { return IsChannelMasked(MediaPlayerEditor::MediaImage::ETextureChannelMask::Red); })
+	);
+
+	ToolkitCommands->MapAction(
+		Commands.ToggledGreenTextureChannel,
+		FExecuteAction::CreateLambda([this] { ToggleChannelMask(MediaPlayerEditor::MediaImage::ETextureChannelMask::Green); }),
+		FCanExecuteAction::CreateLambda([this] { return Viewer.IsValid(); }),
+		FIsActionChecked::CreateLambda([this] { return IsChannelMasked(MediaPlayerEditor::MediaImage::ETextureChannelMask::Green); })
+	);
+
+	ToolkitCommands->MapAction(
+		Commands.ToggledBlueTextureChannel,
+		FExecuteAction::CreateLambda([this] { ToggleChannelMask(MediaPlayerEditor::MediaImage::ETextureChannelMask::Blue); }),
+		FCanExecuteAction::CreateLambda([this] { return Viewer.IsValid(); }),
+		FIsActionChecked::CreateLambda([this] { return IsChannelMasked(MediaPlayerEditor::MediaImage::ETextureChannelMask::Blue); })
+	);
+
+	ToolkitCommands->MapAction(
+		Commands.ToggledAlphaTextureChannel,
+		FExecuteAction::CreateLambda([this] { ToggleChannelMask(MediaPlayerEditor::MediaImage::ETextureChannelMask::Alpha); }),
+		FCanExecuteAction::CreateLambda([this] { return Viewer.IsValid(); }),
+		FIsActionChecked::CreateLambda([this] { return IsChannelMasked(MediaPlayerEditor::MediaImage::ETextureChannelMask::Alpha); })
+	);
 }
 
 BEGIN_SLATE_FUNCTION_BUILD_OPTIMIZATION
 void FMediaPlateEditorToolkit::ExtendToolBar()
 {
-	struct Local
-	{
-		static void FillToolbar(FToolBarBuilder& ToolbarBuilder, const TSharedRef<FUICommandList> ToolkitCommands)
-		{
-			const FMediaPlayerEditorCommands& Commands = FMediaPlayerEditorCommands::GetExternal();
-
-			ToolbarBuilder.BeginSection("PlaybackControls");
-			{
-				ToolbarBuilder.AddToolBarButton(Commands.PreviousMedia);
-				ToolbarBuilder.AddToolBarButton(Commands.RewindMedia);
-				ToolbarBuilder.AddToolBarButton(Commands.ReverseMedia);
-				ToolbarBuilder.AddToolBarButton(Commands.PlayMedia);
-				ToolbarBuilder.AddToolBarButton(Commands.PauseMedia);
-				ToolbarBuilder.AddToolBarButton(Commands.ForwardMedia);
-				ToolbarBuilder.AddToolBarButton(Commands.NextMedia);
-			}
-			ToolbarBuilder.EndSection();
-
-			ToolbarBuilder.BeginSection("MediaControls");
-			{
-				ToolbarBuilder.AddToolBarButton(Commands.OpenMedia);
-				ToolbarBuilder.AddToolBarButton(Commands.CloseMedia);
-			}
-			ToolbarBuilder.EndSection();
-		}
-	};
-
 	TSharedPtr<FExtender> ToolbarExtender = MakeShareable(new FExtender);
 
 	ToolbarExtender->AddToolBarExtension(
 		"Asset",
 		EExtensionHook::After,
 		GetToolkitCommands(),
-		FToolBarExtensionDelegate::CreateStatic(&Local::FillToolbar, GetToolkitCommands())
+		FToolBarExtensionDelegate::CreateLambda(
+			[this](FToolBarBuilder& ToolbarBuilder)
+			{
+				const FMediaPlayerEditorCommands& Commands = FMediaPlayerEditorCommands::GetExternal();
+
+				ToolbarBuilder.BeginSection("PlaybackControls");
+				{
+					ToolbarBuilder.AddToolBarButton(Commands.PreviousMedia);
+					ToolbarBuilder.AddToolBarButton(Commands.RewindMedia);
+					ToolbarBuilder.AddToolBarButton(Commands.ReverseMedia);
+					ToolbarBuilder.AddToolBarButton(Commands.PlayMedia);
+					ToolbarBuilder.AddToolBarButton(Commands.PauseMedia);
+					ToolbarBuilder.AddToolBarButton(Commands.ForwardMedia);
+					ToolbarBuilder.AddToolBarButton(Commands.NextMedia);
+				}
+				ToolbarBuilder.EndSection();
+
+				ToolbarBuilder.BeginSection("MediaControls");
+				{
+					ToolbarBuilder.AddToolBarButton(Commands.OpenMedia);
+					ToolbarBuilder.AddToolBarButton(Commands.CloseMedia);
+				}
+				ToolbarBuilder.EndSection();
+
+				using namespace MediaPlayerEditor::MediaImage;
+
+				ToolbarBuilder.BeginSection("TextureControls");
+				{
+					ToolbarBuilder.AddWidget(SNew(SMediaImageTextureChannelToggle, Viewer, ETextureChannelMask::Red, ToolkitCommands));
+					ToolbarBuilder.AddWidget(SNew(SMediaImageTextureChannelToggle, Viewer, ETextureChannelMask::Green, ToolkitCommands));
+					ToolbarBuilder.AddWidget(SNew(SMediaImageTextureChannelToggle, Viewer, ETextureChannelMask::Blue, ToolkitCommands));
+					ToolbarBuilder.AddWidget(SNew(SMediaImageTextureChannelToggle, Viewer, ETextureChannelMask::Alpha, ToolkitCommands));
+				}
+				ToolbarBuilder.EndSection();
+			}
+		)
 	);
 
 	AddToolbarExtender(ToolbarExtender);
@@ -505,7 +564,7 @@ TSharedRef<SDockTab> FMediaPlateEditorToolkit::HandleTabManagerSpawnTab(const FS
 	{
 		if (MediaPlayer != nullptr)
 		{
-			TabWidget = SNew(SMediaPlayerEditorViewer, *MediaPlayer, MediaTexture, Style, false)
+			TabWidget = SAssignNew(Viewer, SMediaPlayerEditorViewer, *MediaPlayer, MediaTexture, Style, false)
 				.Commands(ToolkitCommands);
 		}
 	}
@@ -526,4 +585,14 @@ void FMediaPlateEditorToolkit::OnActorDeleted(AActor* Actor)
 	}
 }
 
+bool FMediaPlateEditorToolkit::IsChannelMasked(MediaPlayerEditor::MediaImage::ETextureChannelMask InChannel) const
+{
+	if (Viewer.IsValid())
+	{
+		return EnumHasAllFlags(Viewer->GetChannelMask(), InChannel);
+	}
+
+	return true;
+}
+
 #undef LOCTEXT_NAMESPACE
diff --git a/Engine/Plugins/Media/MediaPlate/Source/MediaPlateEditor/Private/Toolkits/MediaPlateEditorToolkit.h b/Engine/Plugins/Media/MediaPlate/Source/MediaPlateEditor/Private/Toolkits/MediaPlateEditorToolkit.h
--- a/Engine/Plugins/Media/MediaPlate/Source/MediaPlateEditor/Private/Toolkits/MediaPlateEditorToolkit.h
+++ b/Engine/Plugins/Media/MediaPlate/Source/MediaPlateEditor/Private/Toolkits/MediaPlateEditorToolkit.h
@@ -8,9 +8,11 @@
 #include "Styling/ISlateStyle.h"
 #include "Toolkits/AssetEditorToolkit.h"
 #include "UObject/GCObject.h"
+#include "Widgets/SMediaImage.h"
 
-enum class EMediaPlateEventState : uint8;
+class SMediaPlayerEditorViewer;
 class UMediaPlateComponent;
+enum class EMediaPlateEventState : uint8;
 
 /**
  * Implements an Editor toolkit for media plates.
@@ -39,6 +41,16 @@ public:
 	 */
 	void Initialize(UMediaPlateComponent* InMediaPlate, const EToolkitMode::Type InMode, const TSharedPtr<IToolkitHost>& InToolkitHost);
 
+	/**
+	 * Returns true if all the given channel mask is enabled.
+	 */
+	bool IsChannelMasked(MediaPlayerEditor::MediaImage::ETextureChannelMask InChannel) const;
+
+	/**
+	 * Toggles the given mask in the channel mask.
+	 */
+	void ToggleChannelMask(MediaPlayerEditor::MediaImage::ETextureChannelMask InChannelToToggle);
+
 	//~ FAssetEditorToolkit interface
 	virtual FString GetDocumentationLink() const override;
 	virtual void OnClose() override;
@@ -85,6 +97,15 @@ protected:
 	float GetReverseRate() const;
 
 private:
+	/** The media plate asset being edited. */
+	TObjectPtr<UMediaPlateComponent> MediaPlate;
+
+	/** Pointer to the style set to use for toolkits. */
+	TSharedRef<ISlateStyle> Style;
+
+	/** The tab which displays the media. */
+	TSharedPtr<SMediaPlayerEditorViewer> Viewer;
+
 	/** Callback to validate playback button events. */
 	bool CanHandleMediaPlateEvent(EMediaPlateEventState InState);
 
@@ -99,12 +120,6 @@ private:
 
 	/** Callback for spawning tabs. */
 	TSharedRef<SDockTab> HandleTabManagerSpawnTab(const FSpawnTabArgs& Args, FName TabIdentifier);
-
-	/** The media plate asset being edited. */
-	TObjectPtr<UMediaPlateComponent> MediaPlate;
-
-	/** Pointer to the style set to use for toolkits. */
-	TSharedRef<ISlateStyle> Style;
 	
 	void OnActorDeleted(AActor* Actor);
 };
diff --git a/Engine/Plugins/Media/MediaPlayerEditor/Source/MediaPlayerEditor/Private/Models/MediaPlayerEditorCommands.cpp b/Engine/Plugins/Media/MediaPlayerEditor/Source/MediaPlayerEditor/Private/Models/MediaPlayerEditorCommands.cpp
--- a/Engine/Plugins/Media/MediaPlayerEditor/Source/MediaPlayerEditor/Private/Models/MediaPlayerEditorCommands.cpp
+++ b/Engine/Plugins/Media/MediaPlayerEditor/Source/MediaPlayerEditor/Private/Models/MediaPlayerEditorCommands.cpp
@@ -90,6 +90,38 @@ void FMediaPlayerEditorCommands::RegisterCommands()
 		EUserInterfaceActionType::Button, 
 		FInputChord(EKeys::L, EModifierKey::Control)
 	);
+
+	UI_COMMAND(
+		ToggledRedTextureChannel,
+		"R",
+		"Toggles the red texture channel on and onn",
+		EUserInterfaceActionType::ToggleButton,
+		FInputChord(EKeys::R)
+	);
+
+	UI_COMMAND(
+		ToggledGreenTextureChannel,
+		"G",
+		"Toggles the green texture channel on and onn",
+		EUserInterfaceActionType::ToggleButton,
+		FInputChord(EKeys::G)
+	);
+
+	UI_COMMAND(
+		ToggledBlueTextureChannel,
+		"B",
+		"Toggles the blue texture channel on and onn",
+		EUserInterfaceActionType::ToggleButton,
+		FInputChord(EKeys::B)
+	);
+
+	UI_COMMAND(
+		ToggledAlphaTextureChannel,
+		"A",
+		"Toggles the alpha texture channel on and onn",
+		EUserInterfaceActionType::ToggleButton,
+		FInputChord(EKeys::A)
+	);
 }
 
 
diff --git a/Engine/Plugins/Media/MediaPlayerEditor/Source/MediaPlayerEditor/Private/Toolkits/MediaPlayerEditorToolkit.cpp b/Engine/Plugins/Media/MediaPlayerEditor/Source/MediaPlayerEditor/Private/Toolkits/MediaPlayerEditorToolkit.cpp
--- a/Engine/Plugins/Media/MediaPlayerEditor/Source/MediaPlayerEditor/Private/Toolkits/MediaPlayerEditorToolkit.cpp
+++ b/Engine/Plugins/Media/MediaPlayerEditor/Source/MediaPlayerEditor/Private/Toolkits/MediaPlayerEditorToolkit.cpp
@@ -8,6 +8,7 @@
 #include "Models/MediaPlayerEditorCommands.h"
 #include "SlateOptMacros.h"
 #include "Widgets/Docking/SDockTab.h"
+#include "Widgets/SMediaImageTextureChannelToggle.h"
 #include "Widgets/SMediaPlayerEditorDetails.h"
 #include "Widgets/SMediaPlayerEditorInfo.h"
 #include "Widgets/SMediaPlayerEditorMedia.h"
@@ -180,37 +181,47 @@ FString FMediaPlayerEditorToolkit::GetWorldCentricTabPrefix() const
 BEGIN_SLATE_FUNCTION_BUILD_OPTIMIZATION
 void FMediaPlayerEditorToolkit::ExtendToolBar()
 {
-	struct Local
-	{
-		static void FillToolbar(FToolBarBuilder& ToolbarBuilder, const TSharedRef<FUICommandList> ToolkitCommands)
-		{
-			ToolbarBuilder.BeginSection("PlaybackControls");
-			{
-				ToolbarBuilder.AddToolBarButton(FMediaPlayerEditorCommands::Get().PreviousMedia);
-				ToolbarBuilder.AddToolBarButton(FMediaPlayerEditorCommands::Get().RewindMedia);
-				ToolbarBuilder.AddToolBarButton(FMediaPlayerEditorCommands::Get().ReverseMedia);
-				ToolbarBuilder.AddToolBarButton(FMediaPlayerEditorCommands::Get().PlayMedia);
-				ToolbarBuilder.AddToolBarButton(FMediaPlayerEditorCommands::Get().PauseMedia);
-				ToolbarBuilder.AddToolBarButton(FMediaPlayerEditorCommands::Get().ForwardMedia);
-				ToolbarBuilder.AddToolBarButton(FMediaPlayerEditorCommands::Get().NextMedia);
-			}
-			ToolbarBuilder.EndSection();
-
-			ToolbarBuilder.BeginSection("MediaControls");
-			{
-				ToolbarBuilder.AddToolBarButton(FMediaPlayerEditorCommands::Get().CloseMedia);
-			}
-			ToolbarBuilder.EndSection();
-		}
-	};
-
-	TSharedPtr<FExtender> ToolbarExtender = MakeShareable(new FExtender);
+	TSharedPtr<FExtender> ToolbarExtender = MakeShared<FExtender>();
 
 	ToolbarExtender->AddToolBarExtension(
 		"Asset",
 		EExtensionHook::After,
 		GetToolkitCommands(),
-		FToolBarExtensionDelegate::CreateStatic(&Local::FillToolbar, GetToolkitCommands())
+		FToolBarExtensionDelegate::CreateLambda(
+			[this](FToolBarBuilder& ToolbarBuilder)
+			{
+				const FMediaPlayerEditorCommands& Commands = FMediaPlayerEditorCommands::Get();
+
+				ToolbarBuilder.BeginSection("PlaybackControls");
+				{
+					ToolbarBuilder.AddToolBarButton(Commands.PreviousMedia);
+					ToolbarBuilder.AddToolBarButton(Commands.RewindMedia);
+					ToolbarBuilder.AddToolBarButton(Commands.ReverseMedia);
+					ToolbarBuilder.AddToolBarButton(Commands.PlayMedia);
+					ToolbarBuilder.AddToolBarButton(Commands.PauseMedia);
+					ToolbarBuilder.AddToolBarButton(Commands.ForwardMedia);
+					ToolbarBuilder.AddToolBarButton(Commands.NextMedia);
+				}
+				ToolbarBuilder.EndSection();
+
+				ToolbarBuilder.BeginSection("MediaControls");
+				{
+					ToolbarBuilder.AddToolBarButton(Commands.CloseMedia);
+				}
+				ToolbarBuilder.EndSection();
+
+				using namespace MediaPlayerEditor::MediaImage;
+
+				ToolbarBuilder.BeginSection("TextureControls");
+				{
+					ToolbarBuilder.AddWidget(SNew(SMediaImageTextureChannelToggle, Viewer, ETextureChannelMask::Red, ToolkitCommands));
+					ToolbarBuilder.AddWidget(SNew(SMediaImageTextureChannelToggle, Viewer, ETextureChannelMask::Green, ToolkitCommands));
+					ToolbarBuilder.AddWidget(SNew(SMediaImageTextureChannelToggle, Viewer, ETextureChannelMask::Blue, ToolkitCommands));
+					ToolbarBuilder.AddWidget(SNew(SMediaImageTextureChannelToggle, Viewer, ETextureChannelMask::Alpha, ToolkitCommands));
+				}
+				ToolbarBuilder.EndSection();
+			}
+		)
 	);
 
 	AddToolbarExtender(ToolbarExtender);
@@ -243,7 +254,7 @@ TSharedRef<SDockTab> FMediaPlayerEditorToolkit::HandleTabManagerSpawnTab(const F
 	}
 	else if (TabIdentifier == MediaPlayerEditorToolkit::ViewerTabId)
 	{
-		TabWidget = SNew(SMediaPlayerEditorViewer, *MediaPlayer, nullptr, Style, true)
+		TabWidget = SAssignNew(Viewer, SMediaPlayerEditorViewer, *MediaPlayer, nullptr, Style, true)
 			.Commands(ToolkitCommands);
 	}
 
diff --git a/Engine/Plugins/Media/MediaPlayerEditor/Source/MediaPlayerEditor/Private/Toolkits/MediaPlayerEditorToolkitMediaPlayerBase.cpp b/Engine/Plugins/Media/MediaPlayerEditor/Source/MediaPlayerEditor/Private/Toolkits/MediaPlayerEditorToolkitMediaPlayerBase.cpp
--- a/Engine/Plugins/Media/MediaPlayerEditor/Source/MediaPlayerEditor/Private/Toolkits/MediaPlayerEditorToolkitMediaPlayerBase.cpp
+++ b/Engine/Plugins/Media/MediaPlayerEditor/Source/MediaPlayerEditor/Private/Toolkits/MediaPlayerEditorToolkitMediaPlayerBase.cpp
@@ -3,10 +3,13 @@
 #include "Toolkits/MediaPlayerEditorToolkitMediaPlayerBase.h"
 
 #include "Framework/Commands/UICommandList.h"
+#include "Framework/MultiBox/MultiBoxBuilder.h"
 #include "MediaPlayer.h"
 #include "MediaPlaylist.h"
 #include "Models/MediaPlayerEditorCommands.h"
 #include "UObject/UObjectGlobals.h"
+#include "Widgets/Input/SCheckBox.h"
+#include "Widgets/SMediaPlayerEditorViewer.h"
 
 FMediaPlayerEditorToolkitMediaPlayerBase::FMediaPlayerEditorToolkitMediaPlayerBase(const TSharedRef<ISlateStyle>& InStyle)
 	: FMediaPlayerEditorToolkitBase(InStyle)
@@ -14,6 +17,42 @@ FMediaPlayerEditorToolkitMediaPlayerBase::FMediaPlayerEditorToolkitMediaPlayerBa
 {
 }
 
+MediaPlayerEditor::MediaImage::ETextureChannelMask FMediaPlayerEditorToolkitMediaPlayerBase::GetChannelMask() const
+{
+	if (Viewer.IsValid())
+	{
+		return Viewer->GetChannelMask();
+	}
+
+	return MediaPlayerEditor::MediaImage::ETextureChannelMask::RGBA;
+}
+
+bool FMediaPlayerEditorToolkitMediaPlayerBase::IsChannelMasked(MediaPlayerEditor::MediaImage::ETextureChannelMask InChannel) const
+{
+	if (Viewer.IsValid())
+	{
+		return Viewer->IsChannelMasked(InChannel);
+	}
+
+	return true;
+}
+
+void FMediaPlayerEditorToolkitMediaPlayerBase::ToggleChannelMask(MediaPlayerEditor::MediaImage::ETextureChannelMask InChannelToToggle)
+{
+	if (Viewer.IsValid())
+	{
+		Viewer->ToggleChannelMask(InChannelToToggle);
+	}
+}
+
+void FMediaPlayerEditorToolkitMediaPlayerBase::SetChannelMask(MediaPlayerEditor::MediaImage::ETextureChannelMask InMask)
+{
+	if (Viewer.IsValid())
+	{
+		Viewer->SetChannelMask(InMask);
+	}
+}
+
 void FMediaPlayerEditorToolkitMediaPlayerBase::OnClose()
 {
 	if (MediaPlayer != nullptr)
@@ -198,4 +237,32 @@ void FMediaPlayerEditorToolkitMediaPlayerBase::BindCommands()
 			}),
 		FCanExecuteAction::CreateLambda([this] { return MediaPlayer->IsReady() && MediaPlayer->SupportsSeeking(); })
 	);
+
+	ToolkitCommands->MapAction(
+		Commands.ToggledRedTextureChannel,
+		FExecuteAction::CreateLambda([this] { ToggleChannelMask(MediaPlayerEditor::MediaImage::ETextureChannelMask::Red); }),
+		FCanExecuteAction::CreateLambda([this] { return Viewer.IsValid(); }),
+		FIsActionChecked::CreateLambda([this] { return IsChannelMasked(MediaPlayerEditor::MediaImage::ETextureChannelMask::Red); })
+	);
+
+	ToolkitCommands->MapAction(
+		Commands.ToggledGreenTextureChannel,
+		FExecuteAction::CreateLambda([this] { ToggleChannelMask(MediaPlayerEditor::MediaImage::ETextureChannelMask::Green); }),
+		FCanExecuteAction::CreateLambda([this] { return Viewer.IsValid(); }),
+		FIsActionChecked::CreateLambda([this] { return IsChannelMasked(MediaPlayerEditor::MediaImage::ETextureChannelMask::Green); })
+	);
+
+	ToolkitCommands->MapAction(
+		Commands.ToggledBlueTextureChannel,
+		FExecuteAction::CreateLambda([this] { ToggleChannelMask(MediaPlayerEditor::MediaImage::ETextureChannelMask::Blue); }),
+		FCanExecuteAction::CreateLambda([this] { return Viewer.IsValid(); }),
+		FIsActionChecked::CreateLambda([this] { return IsChannelMasked(MediaPlayerEditor::MediaImage::ETextureChannelMask::Blue); })
+	);
+
+	ToolkitCommands->MapAction(
+		Commands.ToggledAlphaTextureChannel,
+		FExecuteAction::CreateLambda([this] { ToggleChannelMask(MediaPlayerEditor::MediaImage::ETextureChannelMask::Alpha); }),
+		FCanExecuteAction::CreateLambda([this] { return Viewer.IsValid(); }),
+		FIsActionChecked::CreateLambda([this] { return IsChannelMasked(MediaPlayerEditor::MediaImage::ETextureChannelMask::Alpha); })
+	);
 }
diff --git a/Engine/Plugins/Media/MediaPlayerEditor/Source/MediaPlayerEditor/Private/Toolkits/MediaPlayerEditorToolkitMediaPlayerBase.h b/Engine/Plugins/Media/MediaPlayerEditor/Source/MediaPlayerEditor/Private/Toolkits/MediaPlayerEditorToolkitMediaPlayerBase.h
--- a/Engine/Plugins/Media/MediaPlayerEditor/Source/MediaPlayerEditor/Private/Toolkits/MediaPlayerEditorToolkitMediaPlayerBase.h
+++ b/Engine/Plugins/Media/MediaPlayerEditor/Source/MediaPlayerEditor/Private/Toolkits/MediaPlayerEditorToolkitMediaPlayerBase.h
@@ -4,6 +4,10 @@
 
 #include "Toolkits/MediaPlayerEditorToolkitBase.h"
 
+#include "Widgets/SMediaImage.h"
+
+class FToolBarBuilder;
+class SMediaPlayerEditorViewer;
 class UMediaPlayer;
 
 #define UE_API MEDIAPLAYEREDITOR_API
@@ -22,6 +26,30 @@ public:
 	 */
 	FMediaPlayerEditorToolkitMediaPlayerBase(const TSharedRef<ISlateStyle>& InStyle);
 
+	/**
+	 * Returns the active mask.
+	 */
+	MediaPlayerEditor::MediaImage::ETextureChannelMask GetChannelMask() const;
+
+	/**
+	 * Returns true if all the given channel mask is enabled.
+	 */
+	bool IsChannelMasked(MediaPlayerEditor::MediaImage::ETextureChannelMask InChannel) const;
+
+	/**
+	 * Sets a channel-based mask for the image.
+	 *
+	 * If only a single channel is displayed, it will show in greyscale.
+	 *
+	 * @param InMask to display.
+	 */
+	void SetChannelMask(MediaPlayerEditor::MediaImage::ETextureChannelMask InMask);
+
+	/**
+	 * Toggles the given mask in the channel mask.
+	 */
+	void ToggleChannelMask(MediaPlayerEditor::MediaImage::ETextureChannelMask InChannelToToggle);
+
 	//~ Begin FAssetEditorToolkit
 	virtual void OnClose() override;
 	//~ End FAssetEditorToolkit
@@ -34,6 +62,9 @@ protected:
 	/** The media player to play the media with. */
 	TObjectPtr<UMediaPlayer> MediaPlayer;
 
+	/** The tab which displays the media. */
+	TSharedPtr<SMediaPlayerEditorViewer> Viewer;
+
 	/**
 	 * Gets the playback rate for fast forward.
 	 *
diff --git a/Engine/Plugins/Media/MediaPlayerEditor/Source/MediaPlayerEditor/Private/Toolkits/MediaSourceEditorToolkit.cpp b/Engine/Plugins/Media/MediaPlayerEditor/Source/MediaPlayerEditor/Private/Toolkits/MediaSourceEditorToolkit.cpp
--- a/Engine/Plugins/Media/MediaPlayerEditor/Source/MediaPlayerEditor/Private/Toolkits/MediaSourceEditorToolkit.cpp
+++ b/Engine/Plugins/Media/MediaPlayerEditor/Source/MediaPlayerEditor/Private/Toolkits/MediaSourceEditorToolkit.cpp
@@ -26,6 +26,7 @@
 #include "SlateOptMacros.h"
 #include "TextureResource.h"
 #include "Widgets/Docking/SDockTab.h"
+#include "Widgets/SMediaImageTextureChannelToggle.h"
 #include "Widgets/SMediaPlayerEditorDetails.h"
 #include "Widgets/SMediaPlayerEditorMediaDetails.h"
 #include "Widgets/SMediaPlayerEditorViewer.h"
@@ -234,40 +235,49 @@ void FMediaSourceEditorToolkit::BindCommands()
 BEGIN_SLATE_FUNCTION_BUILD_OPTIMIZATION
 void FMediaSourceEditorToolkit::ExtendToolBar()
 {
-	struct Local
-	{
-		static void FillToolbar(FToolBarBuilder& ToolbarBuilder, const TSharedRef<FUICommandList> ToolkitCommands)
-		{
-			ToolbarBuilder.BeginSection("PlaybackControls");
-			{
-				ToolbarBuilder.AddToolBarButton(FMediaPlayerEditorCommands::Get().PreviousMedia);
-				ToolbarBuilder.AddToolBarButton(FMediaPlayerEditorCommands::Get().RewindMedia);
-				ToolbarBuilder.AddToolBarButton(FMediaPlayerEditorCommands::Get().ReverseMedia);
-				ToolbarBuilder.AddToolBarButton(FMediaPlayerEditorCommands::Get().PlayMedia);
-				ToolbarBuilder.AddToolBarButton(FMediaPlayerEditorCommands::Get().PauseMedia);
-				ToolbarBuilder.AddToolBarButton(FMediaPlayerEditorCommands::Get().ForwardMedia);
-				ToolbarBuilder.AddToolBarButton(FMediaPlayerEditorCommands::Get().NextMedia);
-			}
-			ToolbarBuilder.EndSection();
-
-			ToolbarBuilder.BeginSection("MediaControls");
-			{
-				ToolbarBuilder.AddToolBarButton(FMediaPlayerEditorCommands::Get().OpenMedia);
-				ToolbarBuilder.AddToolBarButton(FMediaPlayerEditorCommands::Get().CloseMedia);
-				ToolbarBuilder.AddToolBarButton(FMediaPlayerEditorCommands::Get().GenerateThumbnail);
-			}
-			ToolbarBuilder.EndSection();
-		}
-	};
-
-
-	TSharedPtr<FExtender> ToolbarExtender = MakeShareable(new FExtender);
+	TSharedPtr<FExtender> ToolbarExtender = MakeShared<FExtender>();
 
 	ToolbarExtender->AddToolBarExtension(
 		"Asset",
 		EExtensionHook::After,
 		GetToolkitCommands(),
-		FToolBarExtensionDelegate::CreateStatic(&Local::FillToolbar, GetToolkitCommands())
+		FToolBarExtensionDelegate::CreateLambda(
+			[this](FToolBarBuilder& ToolbarBuilder)
+			{
+				const FMediaPlayerEditorCommands& Commands = FMediaPlayerEditorCommands::Get();
+
+				ToolbarBuilder.BeginSection("PlaybackControls");
+				{
+					ToolbarBuilder.AddToolBarButton(Commands.PreviousMedia);
+					ToolbarBuilder.AddToolBarButton(Commands.RewindMedia);
+					ToolbarBuilder.AddToolBarButton(Commands.ReverseMedia);
+					ToolbarBuilder.AddToolBarButton(Commands.PlayMedia);
+					ToolbarBuilder.AddToolBarButton(Commands.PauseMedia);
+					ToolbarBuilder.AddToolBarButton(Commands.ForwardMedia);
+					ToolbarBuilder.AddToolBarButton(Commands.NextMedia);
+				}
+				ToolbarBuilder.EndSection();
+
+				ToolbarBuilder.BeginSection("MediaControls");
+				{
+					ToolbarBuilder.AddToolBarButton(Commands.OpenMedia);
+					ToolbarBuilder.AddToolBarButton(Commands.CloseMedia);
+					ToolbarBuilder.AddToolBarButton(Commands.GenerateThumbnail);
+				}
+				ToolbarBuilder.EndSection();
+
+				using namespace MediaPlayerEditor::MediaImage;
+
+				ToolbarBuilder.BeginSection("TextureControls");
+				{
+					ToolbarBuilder.AddWidget(SNew(SMediaImageTextureChannelToggle, Viewer, ETextureChannelMask::Red, ToolkitCommands));
+					ToolbarBuilder.AddWidget(SNew(SMediaImageTextureChannelToggle, Viewer, ETextureChannelMask::Green, ToolkitCommands));
+					ToolbarBuilder.AddWidget(SNew(SMediaImageTextureChannelToggle, Viewer, ETextureChannelMask::Blue, ToolkitCommands));
+					ToolbarBuilder.AddWidget(SNew(SMediaImageTextureChannelToggle, Viewer, ETextureChannelMask::Alpha, ToolkitCommands));
+				}
+				ToolbarBuilder.EndSection();
+			}
+		)
 	);
 
 	AddToolbarExtender(ToolbarExtender);
@@ -370,7 +380,7 @@ TSharedRef<SDockTab> FMediaSourceEditorToolkit::HandleTabManagerSpawnTab(const F
 	}
 	else if (TabIdentifier == MediaSourceEditorToolkit::ViewerTabId)
 	{
-		TabWidget = SNew(SMediaPlayerEditorViewer, *MediaPlayer, MediaTexture, Style, true)
+		TabWidget = SAssignNew(Viewer, SMediaPlayerEditorViewer, *MediaPlayer, MediaTexture, Style, true)
 			.bShowUrl(false)
 			.Commands(ToolkitCommands);
 	}
diff --git a/Engine/Plugins/Media/MediaPlayerEditor/Source/MediaPlayerEditor/Private/Widgets/SMediaImage.cpp b/Engine/Plugins/Media/MediaPlayerEditor/Source/MediaPlayerEditor/Private/Widgets/SMediaImage.cpp
--- a/Engine/Plugins/Media/MediaPlayerEditor/Source/MediaPlayerEditor/Private/Widgets/SMediaImage.cpp
+++ b/Engine/Plugins/Media/MediaPlayerEditor/Source/MediaPlayerEditor/Private/Widgets/SMediaImage.cpp
@@ -2,43 +2,87 @@
 
 #include "Widgets/SMediaImage.h"
 
-#include "Styling/AppStyle.h"
+#include "Engine/Texture.h"
 #include "MaterialDomain.h"
 #include "Materials/Material.h"
-#include "Materials/MaterialExpressionTextureSample.h"
+#include "Materials/MaterialExpressionConvert.h"
 #include "Materials/MaterialExpressionSubstrate.h"
+#include "Materials/MaterialExpressionTextureObjectParameter.h"
+#include "Materials/MaterialExpressionTextureSample.h"
+#include "Materials/MaterialInstanceDynamic.h"
 #include "RenderUtils.h"
+#include "Styling/AppStyle.h"
 #include "Styling/SlateBrush.h"
 #include "UObject/Package.h"
-#include "Widgets/Layout/SScaleBox.h"
 #include "Widgets/Images/SImage.h"
-#include "Widgets/Input/SButton.h"
-
-
-SMediaImage::SMediaImage()
-	: Collector(this)
-	, Material(nullptr)
-	, MaterialBrush(nullptr)
-	, TextureSampler(nullptr)
-{ }
+#include "Widgets/Layout/SScaleBox.h"
 
-void SMediaImage::Construct(const FArguments& InArgs, UTexture* InTexture)
+namespace MediaPlayerEditor
 {
-	// The Slate brush that renders the material.
-	BrushImageSize = InArgs._BrushImageSize;
-
-	if (InTexture != nullptr)
+	namespace MediaImage::Private
 	{
-		// create wrapper material
-		Material = NewObject<UMaterial>(GetTransientPackage(), NAME_None, RF_Transient);
+		static const FLazyName TextureParameterName = "Texture";
 
-		if (Material != nullptr)
+		UMaterial* CreateChannelMaskMaterial(UTexture* InTexture, MediaPlayerEditor::MediaImage::ETextureChannelMask InMask)
 		{
+			UMaterial* Material = NewObject<UMaterial>(GetTransientPackage(), NAME_None, RF_Transient);
+
+			if (!Material)
+			{
+				return {};
+			}
+
 			UMaterialEditorOnlyData* MaterialEditorOnly = Material->GetEditorOnlyData();
-			TextureSampler = NewObject<UMaterialExpressionTextureSample>(Material);
+
+			UMaterialExpressionTextureObjectParameter* TextureParameter = NewObject<UMaterialExpressionTextureObjectParameter>(Material);
+			{
+				TextureParameter->SetParameterName(TextureParameterName);
+				TextureParameter->SamplerSource = ESamplerSourceMode::SSM_Clamp_WorldGroupSettings;
+			}
+
+			UMaterialExpressionTextureSample* TextureSampler = NewObject<UMaterialExpressionTextureSample>(Material);
 			{
 				TextureSampler->Texture = InTexture;
 				TextureSampler->AutoSetSampleType();
+				TextureSampler->Texture = nullptr;
+
+				TextureSampler->TextureObject.Connect(0, TextureParameter);
+			}
+
+			using namespace MediaPlayerEditor::MediaImage;
+
+			UMaterialExpressionConvert* ConvertNode = NewObject<UMaterialExpressionConvert>(Material);
+			{
+				FMaterialExpressionConvertInput& RedInput = ConvertNode->ConvertInputs.AddDefaulted_GetRef();
+				RedInput.Type = EMaterialExpressionConvertType::Scalar;
+				
+				if (EnumHasAnyFlags(InMask, ETextureChannelMask::Red))
+				{
+					RedInput.ExpressionInput.Connect(1, TextureSampler);
+				}
+
+				FMaterialExpressionConvertInput& GreenInput = ConvertNode->ConvertInputs.AddDefaulted_GetRef();
+				GreenInput.Type = EMaterialExpressionConvertType::Scalar;
+
+				if (EnumHasAnyFlags(InMask, ETextureChannelMask::Green))
+				{
+					GreenInput.ExpressionInput.Connect(2, TextureSampler);
+				}
+
+				FMaterialExpressionConvertInput& BlueInput = ConvertNode->ConvertInputs.AddDefaulted_GetRef();
+				BlueInput.Type = EMaterialExpressionConvertType::Scalar;
+
+				if (EnumHasAnyFlags(InMask, ETextureChannelMask::Blue))
+				{
+					BlueInput.ExpressionInput.Connect(3, TextureSampler);
+				}
+
+				FMaterialExpressionConvertOutput& ConvertOutput = ConvertNode->ConvertOutputs.AddDefaulted_GetRef();
+				ConvertOutput.Type = EMaterialExpressionConvertType::Vector3;
+
+				ConvertNode->ConvertMappings.Add({ 0, 0, 0, 0 });
+				ConvertNode->ConvertMappings.Add({ 1, 0, 0, 1 });
+				ConvertNode->ConvertMappings.Add({ 2, 0, 0, 2 });
 			}
 
 			FExpressionOutput& Output = TextureSampler->GetOutputs()[0];
@@ -47,19 +91,22 @@ void SMediaImage::Construct(const FArguments& InArgs, UTexture* InTexture)
 				UMaterialExpressionSubstrateUI* UINode = NewObject<UMaterialExpressionSubstrateUI>(Material);
 				UINode->Material = Material;
 
-				UINode->Color.Expression = TextureSampler;
-				UINode->Color.Mask = Output.Mask;
-				UINode->Color.MaskR = Output.MaskR;
-				UINode->Color.MaskG = Output.MaskG;
-				UINode->Color.MaskB = Output.MaskB;
-				UINode->Color.MaskA = Output.MaskA;
+				UINode->Color.Expression = ConvertNode;
+				UINode->Color.Mask = 1;
+				UINode->Color.MaskR = 1;
+				UINode->Color.MaskG = 1;
+				UINode->Color.MaskB = 1;
+				UINode->Color.MaskA = 0;
 
-				UINode->Opacity.Expression = TextureSampler;
-				UINode->Opacity.Mask = Output.Mask;
-				UINode->Opacity.MaskR = 0;
-				UINode->Opacity.MaskG = 0;
-				UINode->Opacity.MaskB = 0;
-				UINode->Opacity.MaskA = 1;
+				if (EnumHasAnyFlags(InMask, ETextureChannelMask::Alpha))
+				{
+					UINode->Opacity.Expression = TextureSampler;
+					UINode->Opacity.Mask = 1;
+					UINode->Opacity.MaskR = 0;
+					UINode->Opacity.MaskG = 0;
+					UINode->Opacity.MaskB = 0;
+					UINode->Opacity.MaskA = 1;
+				}
 
 				MaterialEditorOnly->FrontMaterial.Connect(0, UINode);
 			}
@@ -67,18 +114,19 @@ void SMediaImage::Construct(const FArguments& InArgs, UTexture* InTexture)
 			{
 				FExpressionInput& Input = MaterialEditorOnly->EmissiveColor;
 				{
-					Input.Expression = TextureSampler;
-					Input.Mask = Output.Mask;
-					Input.MaskR = Output.MaskR;
-					Input.MaskG = Output.MaskG;
-					Input.MaskB = Output.MaskB;
-					Input.MaskA = Output.MaskA;
+					Input.Expression = ConvertNode;
+					Input.Mask = 1;
+					Input.MaskR = 1;
+					Input.MaskG = 1;
+					Input.MaskB = 1;
+					Input.MaskA = 0;
 				}
 
-				FExpressionInput& Opacity = MaterialEditorOnly->Opacity;
+				if (EnumHasAnyFlags(InMask, ETextureChannelMask::Alpha))
 				{
+					FExpressionInput& Opacity = MaterialEditorOnly->Opacity;
 					Opacity.Expression = TextureSampler;
-					Opacity.Mask = Output.Mask;
+					Opacity.Mask = 1;
 					Opacity.MaskR = 0;
 					Opacity.MaskG = 0;
 					Opacity.MaskB = 0;
@@ -88,17 +136,152 @@ void SMediaImage::Construct(const FArguments& InArgs, UTexture* InTexture)
 
 			Material->BlendMode = BLEND_AlphaComposite;
 
+			Material->GetExpressionCollection().AddExpression(TextureParameter);
 			Material->GetExpressionCollection().AddExpression(TextureSampler);
+			Material->GetExpressionCollection().AddExpression(ConvertNode);
 			Material->MaterialDomain = EMaterialDomain::MD_UI;
 			Material->PostEditChange();
+
+			return Material;
 		}
 
-		// create Slate brush
-		MaterialBrush = MakeShareable(new FSlateBrush());
+		UMaterial* CreateGreyscaleMaterial(UTexture* InTexture, MediaPlayerEditor::MediaImage::ETextureChannelMask InChannel)
+		{
+			UMaterial* Material = NewObject<UMaterial>(GetTransientPackage(), NAME_None, RF_Transient);
+
+			if (!Material)
+			{
+				return {};
+			}
+
+			UMaterialEditorOnlyData* MaterialEditorOnly = Material->GetEditorOnlyData();
+
+			UMaterialExpressionTextureObjectParameter* TextureParameter = NewObject<UMaterialExpressionTextureObjectParameter>(Material);
+			{
+				TextureParameter->SetParameterName(TextureParameterName);
+				TextureParameter->SamplerSource = ESamplerSourceMode::SSM_Clamp_WorldGroupSettings;
+			}
+
+			UMaterialExpressionTextureSample* TextureSampler = NewObject<UMaterialExpressionTextureSample>(Material);
+			{
+				TextureSampler->Texture = InTexture;
+				TextureSampler->AutoSetSampleType();
+				TextureSampler->Texture = nullptr;
+
+				TextureSampler->TextureObject.Connect(0, TextureParameter);
+			}
+
+			using namespace MediaPlayerEditor::MediaImage;
+
+			UMaterialExpressionConvert* ConvertNode = NewObject<UMaterialExpressionConvert>(Material);
+			{
+				FMaterialExpressionConvertInput& ConvertInput = ConvertNode->ConvertInputs.AddDefaulted_GetRef();
+				ConvertInput.Type = EMaterialExpressionConvertType::Vector3;
+
+				switch (InChannel)
+				{
+					case ETextureChannelMask::Red:
+						ConvertInput.ExpressionInput.Connect(1, TextureSampler);
+						break;
+
+					case ETextureChannelMask::Green:
+						ConvertInput.ExpressionInput.Connect(2, TextureSampler);
+						break;
+
+					case ETextureChannelMask::Blue:
+						ConvertInput.ExpressionInput.Connect(3, TextureSampler);
+						break;
+
+					case ETextureChannelMask::Alpha:
+						ConvertInput.ExpressionInput.Connect(4, TextureSampler);
+						break;
+
+					default:
+						return {};
+				}
+
+				FMaterialExpressionConvertOutput& ConvertOutput = ConvertNode->ConvertOutputs.AddDefaulted_GetRef();
+				ConvertOutput.Type = EMaterialExpressionConvertType::Vector3;
+
+				ConvertNode->ConvertMappings.Add({ 0, 0, 0, 0 }); // Input -> Red
+				ConvertNode->ConvertMappings.Add({ 0, 0, 0, 1 }); // Input -> Green
+				ConvertNode->ConvertMappings.Add({ 0, 0, 0, 2 }); // Input -> Blue
+			}		
+
+			if (Substrate::IsSubstrateEnabled())
+			{
+				UMaterialExpressionSubstrateUI* UINode = NewObject<UMaterialExpressionSubstrateUI>(Material);
+				UINode->Material = Material;
+
+				UINode->Color.Expression = ConvertNode;
+				UINode->Color.Mask = 1;
+				UINode->Color.MaskR = 1;
+				UINode->Color.MaskG = 1;
+				UINode->Color.MaskB = 1;
+				UINode->Color.MaskA = 0;
+
+				MaterialEditorOnly->FrontMaterial.Connect(0, UINode);
+			}
+			else
+			{
+				FExpressionInput& Input = MaterialEditorOnly->EmissiveColor;
+				{
+					Input.Expression = ConvertNode;
+					Input.Mask = 1;
+					Input.MaskR = 1;
+					Input.MaskG = 1;
+					Input.MaskB = 1;
+					Input.MaskA = 0;
+				}
+			}
+
+			Material->BlendMode = BLEND_Opaque;
+
+			Material->GetExpressionCollection().AddExpression(TextureParameter);
+			Material->GetExpressionCollection().AddExpression(TextureSampler);
+			Material->GetExpressionCollection().AddExpression(ConvertNode);
+			Material->MaterialDomain = EMaterialDomain::MD_UI;
+			Material->PostEditChange();
+
+			return Material;
+		}
+
+		UMaterialInstanceDynamic* CreateInstance(UMaterial* InBaseMaterial, UTexture* InTexture)
 		{
-			MaterialBrush->SetResourceObject(Material);
+			if (!InBaseMaterial || !InTexture)
+			{
+				return nullptr;
+			}
+
+			UMaterialInstanceDynamic* Instance = UMaterialInstanceDynamic::Create(InBaseMaterial, GetTransientPackage());
+			Instance->SetTextureParameterValue(TextureParameterName, InTexture);
+
+			return Instance;
 		}
 	}
+}
+
+TMap<MediaPlayerEditor::MediaImage::ETextureChannelMask, TStrongObjectPtr<UMaterial>> SMediaImage::CachedMaterials;
+
+SMediaImage::SMediaImage()
+	: Mask(MediaPlayerEditor::MediaImage::ETextureChannelMask::RGBA)
+{}
+
+void SMediaImage::Construct(const FArguments& InArgs, UTexture* InTexture)
+{
+	TextureWeak = InTexture;
+	
+	// The Slate brush that renders the material.
+	BrushImageSize = InArgs._BrushImageSize;
+
+	if (InTexture != nullptr)
+	{
+		// create Slate brush
+		MaterialBrush = MakeShareable(new FSlateBrush());
+
+		// Create material and assign to brush.
+		SetChannelMask(Mask);
+	}
 
 	ChildSlot
 	[
@@ -115,10 +298,78 @@ void SMediaImage::Tick(const FGeometry& InAllottedGeometry, const double InCurre
 {
 	SCompoundWidget::Tick(InAllottedGeometry, InCurrentTime, InDeltaTime);
 
-	if (BrushImageSize.IsSet())
+	if (MaterialBrush.IsValid() && BrushImageSize.IsSet())
 	{
 		FVector2D Size = BrushImageSize.Get();
 		MaterialBrush->ImageSize.X = static_cast<float>(Size.X);
 		MaterialBrush->ImageSize.Y = static_cast<float>(Size.Y);
 	}
 }
+
+MediaPlayerEditor::MediaImage::ETextureChannelMask SMediaImage::GetChannelMask() const
+{
+	return Mask;
+}
+
+void SMediaImage::SetChannelMask(MediaPlayerEditor::MediaImage::ETextureChannelMask InMask)
+{
+	using namespace MediaPlayerEditor::MediaImage;
+
+	UTexture* Texture = TextureWeak.Get();
+
+	if (!Texture)
+	{
+		return;
+	}
+
+	// Invalid, no channels specified
+	if (!EnumHasAnyFlags(InMask, ETextureChannelMask::RGBA))
+	{
+		return;
+	}
+
+	UMaterial* MaskMaterial = nullptr;
+
+	if (const TStrongObjectPtr<UMaterial>* CachedMaterial = CachedMaterials.Find(InMask))
+	{
+		MaskMaterial = CachedMaterial->Get();
+	}
+	else
+	{
+		switch (InMask)
+		{
+			case ETextureChannelMask::Red:
+			case ETextureChannelMask::Green:
+			case ETextureChannelMask::Blue:
+			case ETextureChannelMask::Alpha:
+				// Create greyscale mask
+				MaskMaterial = Private::CreateGreyscaleMaterial(Texture, InMask);
+				break;
+
+			default:
+				// Create channel mask
+				MaskMaterial = Private::CreateChannelMaskMaterial(Texture, InMask);
+				break;
+		}
+
+		if (!MaskMaterial)
+		{
+			return;
+			
+		}
+		
+		CachedMaterials.Add(InMask, TStrongObjectPtr<UMaterial>(MaskMaterial));
+	}
+
+	if (MaskMaterial)
+	{
+		MaterialAInstance.Reset(Private::CreateInstance(MaskMaterial, Texture));
+
+		if (MaterialBrush.IsValid())
+		{
+			MaterialBrush->SetResourceObject(MaterialAInstance.Get());
+		}
+
+		Mask = InMask;
+	}
+}
diff --git a/Engine/Plugins/Media/MediaPlayerEditor/Source/MediaPlayerEditor/Private/Widgets/SMediaImageTextureChannelToggle.cpp b/Engine/Plugins/Media/MediaPlayerEditor/Source/MediaPlayerEditor/Private/Widgets/SMediaImageTextureChannelToggle.cpp
new file mode 100644
--- /dev/null
+++ b/Engine/Plugins/Media/MediaPlayerEditor/Source/MediaPlayerEditor/Private/Widgets/SMediaImageTextureChannelToggle.cpp
@@ -0,0 +1,201 @@
+// Copyright Epic Games, Inc. All Rights Reserved.
+
+#include "Widgets/SMediaImageTextureChannelToggle.h"
+
+#include "Framework/Commands/UICommandList.h"
+#include "MediaPlayer.h"
+#include "Models/MediaPlayerEditorCommands.h"
+#include "Styling/AppStyle.h"
+#include "Widgets/Input/SCheckBox.h"
+#include "Widgets/Text/STextBlock.h"
+#include "Widgets/SMediaPlayerEditorViewer.h"
+
+void SMediaImageTextureChannelToggle::Construct(const FArguments& InArgs, TSharedPtr<SMediaPlayerEditorViewer> InViewer,
+	MediaPlayerEditor::MediaImage::ETextureChannelMask InChannel, TSharedPtr<FUICommandList> InCommandList)
+{
+	ViewerWeak = InViewer;
+	Channel = InChannel;
+	CommandListWeak = InCommandList;
+
+	using namespace MediaPlayerEditor::MediaImage;
+
+	const FMediaPlayerEditorCommands& Commands = FMediaPlayerEditorCommands::Get();
+	TSharedPtr<FUICommandInfo> Command;
+
+	switch (InChannel)
+	{
+		case ETextureChannelMask::Red:
+			Command = Commands.ToggledRedTextureChannel;
+			break;
+
+		case ETextureChannelMask::Green:
+			Command = Commands.ToggledGreenTextureChannel;
+			break;
+
+		case ETextureChannelMask::Blue:
+			Command = Commands.ToggledBlueTextureChannel;
+			break;
+
+		case ETextureChannelMask::Alpha:
+			Command = Commands.ToggledAlphaTextureChannel;
+			break;
+
+		default:
+			check(false);
+			return;
+	}
+
+	ChildSlot
+	[
+		SNew(SCheckBox)
+		.Style(FAppStyle::Get(), "TextureEditor.ChannelButtonStyle")
+		.BorderBackgroundColor(this, &SMediaImageTextureChannelToggle::GetChannelButtonBackgroundColor)
+		.ForegroundColor(this, &SMediaImageTextureChannelToggle::GetChannelButtonForegroundColor)
+		.OnCheckStateChanged(this, &SMediaImageTextureChannelToggle::OnChannelButtonCheckStateChanged)
+		.IsChecked(this, &SMediaImageTextureChannelToggle::OnGetChannelButtonCheckState)
+		.IsEnabled(this, &SMediaImageTextureChannelToggle::IsChannelButtonEnabled)
+		.ToolTipText(Command->GetDescription())
+		.Padding(1.f)
+		[
+			SNew(STextBlock)
+			.Font(FAppStyle::Get().GetFontStyle("TextureEditor.ChannelButtonFont"))
+			.Text(Command->GetLabel())
+		]
+	];
+}
+
+FSlateColor SMediaImageTextureChannelToggle::GetChannelButtonBackgroundColor() const
+{
+	using namespace MediaPlayerEditor::MediaImage;
+
+	FSlateColor Dropdown = FAppStyle::Get().GetSlateColor("Colors.Dropdown");
+
+	TSharedPtr<SMediaPlayerEditorViewer> Viewer = ViewerWeak.Pin();
+
+	if (!Viewer.IsValid())
+	{
+		return FLinearColor::White;
+	}
+
+	switch (Channel)
+	{
+		case ETextureChannelMask::Red:
+			return Viewer->IsChannelMasked(Channel) ? FLinearColor::Red : FLinearColor::White;
+
+		case ETextureChannelMask::Green:
+			return Viewer->IsChannelMasked(Channel) ? FLinearColor::Green : FLinearColor::White;
+
+		case ETextureChannelMask::Blue:
+			return Viewer->IsChannelMasked(Channel) ? FLinearColor::Blue : FLinearColor::White;
+
+		case ETextureChannelMask::Alpha:
+			return FLinearColor::White;
+
+		default:
+			check(false);
+			return FSlateColor();
+	}
+}
+
+FSlateColor SMediaImageTextureChannelToggle::GetChannelButtonForegroundColor() const
+{
+	using namespace MediaPlayerEditor::MediaImage;
+
+	FSlateColor DefaultForeground = FAppStyle::Get().GetSlateColor("Colors.Foreground");
+
+	TSharedPtr<SMediaPlayerEditorViewer> Viewer = ViewerWeak.Pin();
+
+	if (!Viewer.IsValid())
+	{
+		return DefaultForeground;
+	}
+
+	switch (Channel)
+	{
+		case ETextureChannelMask::Red:
+			return Viewer->IsChannelMasked(Channel) ? FLinearColor::Black : DefaultForeground;
+
+		case ETextureChannelMask::Green:
+			return Viewer->IsChannelMasked(Channel) ? FLinearColor::Black : DefaultForeground;
+
+		case ETextureChannelMask::Blue:
+			return Viewer->IsChannelMasked(Channel) ? FLinearColor::Black : DefaultForeground;
+
+		case ETextureChannelMask::Alpha:
+			return Viewer->IsChannelMasked(Channel) ? FLinearColor::Black : DefaultForeground;
+
+		default:
+			check(false);
+			return FSlateColor::UseForeground();
+	}
+}
+
+ECheckBoxState SMediaImageTextureChannelToggle::OnGetChannelButtonCheckState() const
+{
+	if (TSharedPtr<SMediaPlayerEditorViewer> Viewer = ViewerWeak.Pin())
+	{
+		if (Viewer->IsChannelMasked(Channel))
+		{
+			return ECheckBoxState::Checked;
+		}
+	}
+
+	return ECheckBoxState::Unchecked;
+}
+
+bool SMediaImageTextureChannelToggle::IsChannelButtonEnabled() const
+{
+	if (TSharedPtr<SMediaPlayerEditorViewer> Viewer = ViewerWeak.Pin())
+	{
+		if (UMediaPlayer* MediaPlayer = Viewer->GetMediaPlayer())
+		{
+			return MediaPlayer->IsReady();
+		}
+	}
+	
+	return false;
+}
+
+void SMediaImageTextureChannelToggle::OnChannelButtonCheckStateChanged(ECheckBoxState InNewState)
+{
+	using namespace MediaPlayerEditor::MediaImage;
+
+	TSharedPtr<SMediaPlayerEditorViewer> Viewer = ViewerWeak.Pin();
+
+	if (!Viewer.IsValid())
+	{
+		return;
+	}
+
+	TSharedPtr<FUICommandList> CommandList = CommandListWeak.Pin();
+
+	if (!CommandList.IsValid())
+	{
+		return;
+	}
+
+	const FMediaPlayerEditorCommands& Commands = FMediaPlayerEditorCommands::Get();
+
+	switch (Channel)
+	{
+		case ETextureChannelMask::Red:
+			CommandList->TryExecuteAction(Commands.ToggledRedTextureChannel.ToSharedRef());
+			break;
+
+		case ETextureChannelMask::Green:
+			CommandList->TryExecuteAction(Commands.ToggledGreenTextureChannel.ToSharedRef());
+			break;
+
+		case ETextureChannelMask::Blue:
+			CommandList->TryExecuteAction(Commands.ToggledBlueTextureChannel.ToSharedRef());
+			break;
+
+		case ETextureChannelMask::Alpha:
+			CommandList->TryExecuteAction(Commands.ToggledAlphaTextureChannel.ToSharedRef());
+			break;
+
+		default:
+			check(false);
+			break;
+	}
+}
\ No newline at end of file
diff --git a/Engine/Plugins/Media/MediaPlayerEditor/Source/MediaPlayerEditor/Private/Widgets/SMediaPlayerEditorOutput.cpp b/Engine/Plugins/Media/MediaPlayerEditor/Source/MediaPlayerEditor/Private/Widgets/SMediaPlayerEditorOutput.cpp
--- a/Engine/Plugins/Media/MediaPlayerEditor/Source/MediaPlayerEditor/Private/Widgets/SMediaPlayerEditorOutput.cpp
+++ b/Engine/Plugins/Media/MediaPlayerEditor/Source/MediaPlayerEditor/Private/Widgets/SMediaPlayerEditorOutput.cpp
@@ -13,7 +13,6 @@
 #include "MediaTexture.h"
 #include "Styling/SlateBrush.h"
 #include "Widgets/Images/SImage.h"
-#include "Widgets/SMediaImage.h"
 
 
 /* SMediaPlayerEditorOutput structors
@@ -102,7 +101,7 @@ void SMediaPlayerEditorOutput::Construct(const FArguments& InArgs, UMediaPlayer&
 		}
 	}
 
-	TSharedRef<SMediaImage> MediaImage = SNew(SMediaImage, MediaTexture)
+	SAssignNew(MediaImage, SMediaImage, MediaTexture)
 		.BrushImageSize_Lambda([&]()
 			{
 				if (MediaTexture)
@@ -113,12 +112,30 @@ void SMediaPlayerEditorOutput::Construct(const FArguments& InArgs, UMediaPlayer&
 
 	ChildSlot
 	[
-		MediaImage
+		MediaImage.ToSharedRef()
 	];
 
 	MediaPlayer->OnMediaEvent().AddRaw(this, &SMediaPlayerEditorOutput::HandleMediaPlayerMediaEvent);
 }
 
+MediaPlayerEditor::MediaImage::ETextureChannelMask SMediaPlayerEditorOutput::GetChannelMask() const
+{
+	if (MediaImage.IsValid())
+	{
+		return MediaImage->GetChannelMask();
+	}
+
+	return MediaPlayerEditor::MediaImage::ETextureChannelMask::RGBA;
+}
+
+void SMediaPlayerEditorOutput::SetChannelMask(MediaPlayerEditor::MediaImage::ETextureChannelMask InMask)
+{
+	if (MediaImage.IsValid())
+	{
+		MediaImage->SetChannelMask(InMask);
+	}
+}
+
 
 /* SWidget interface
  *****************************************************************************/
diff --git a/Engine/Plugins/Media/MediaPlayerEditor/Source/MediaPlayerEditor/Private/Widgets/SMediaPlayerEditorOutput.h b/Engine/Plugins/Media/MediaPlayerEditor/Source/MediaPlayerEditor/Private/Widgets/SMediaPlayerEditorOutput.h
--- a/Engine/Plugins/Media/MediaPlayerEditor/Source/MediaPlayerEditor/Private/Widgets/SMediaPlayerEditorOutput.h
+++ b/Engine/Plugins/Media/MediaPlayerEditor/Source/MediaPlayerEditor/Private/Widgets/SMediaPlayerEditorOutput.h
@@ -6,6 +6,7 @@
 #include "UObject/WeakObjectPtr.h"
 #include "Widgets/DeclarativeSyntaxSupport.h"
 #include "Widgets/SCompoundWidget.h"
+#include "Widgets/SMediaImage.h"
 
 class UMediaPlayer;
 class UMediaSoundComponent;
@@ -46,6 +47,20 @@ public:
 	void Construct(const FArguments& InArgs, UMediaPlayer& InMediaPlayer,
 		UMediaTexture* InMediaTexture, bool  bInIsSoundEnabled);
 
+	/**
+	 * Returns the active mask.
+	 */
+	MediaPlayerEditor::MediaImage::ETextureChannelMask GetChannelMask() const;
+
+	/**
+	 * Sets a channel-based mask for the image.
+	 *
+	 * If only a single channel is displayed, it will show in greyscale.
+	 *
+	 * @param InMask to display.
+	 */
+	void SetChannelMask(MediaPlayerEditor::MediaImage::ETextureChannelMask InMask);
+
 public:
 
 	//~ SWidget interface
@@ -70,4 +85,7 @@ private:
 
 	/** If true then we created MediaTexture and need to clean it up when we are done. */
 	bool bIsOurMediaTexture;
+
+	/** The image widget displaying the media. */
+	TSharedPtr<SMediaImage> MediaImage;
 };
diff --git a/Engine/Plugins/Media/MediaPlayerEditor/Source/MediaPlayerEditor/Private/Widgets/SMediaPlayerEditorViewer.cpp b/Engine/Plugins/Media/MediaPlayerEditor/Source/MediaPlayerEditor/Private/Widgets/SMediaPlayerEditorViewer.cpp
--- a/Engine/Plugins/Media/MediaPlayerEditor/Source/MediaPlayerEditor/Private/Widgets/SMediaPlayerEditorViewer.cpp
+++ b/Engine/Plugins/Media/MediaPlayerEditor/Source/MediaPlayerEditor/Private/Widgets/SMediaPlayerEditorViewer.cpp
@@ -512,6 +512,11 @@ void SMediaPlayerEditorViewer::Construct(const FArguments& InArgs, UMediaPlayer&
 	];
 }
 
+UMediaPlayer* SMediaPlayerEditorViewer::GetMediaPlayer() const
+{
+	return MediaPlayer;
+}
+
 void SMediaPlayerEditorViewer::EnableMouseControl(bool bIsEnabled)
 {
 	if (PlayerViewport.IsValid())
@@ -520,6 +525,55 @@ void SMediaPlayerEditorViewer::EnableMouseControl(bool bIsEnabled)
 	}
 }
 
+MediaPlayerEditor::MediaImage::ETextureChannelMask SMediaPlayerEditorViewer::GetChannelMask() const
+{
+	if (PlayerViewport.IsValid())
+	{
+		return PlayerViewport->GetChannelMask();
+	}
+
+	return MediaPlayerEditor::MediaImage::ETextureChannelMask::RGBA;
+}
+
+bool SMediaPlayerEditorViewer::IsChannelMasked(MediaPlayerEditor::MediaImage::ETextureChannelMask InChannel) const
+{
+	if (PlayerViewport.IsValid())
+	{
+		return EnumHasAllFlags(PlayerViewport->GetChannelMask(), InChannel);
+	}
+
+	return true;
+}
+
+void SMediaPlayerEditorViewer::SetChannelMask(MediaPlayerEditor::MediaImage::ETextureChannelMask InMask)
+{
+	if (PlayerViewport.IsValid())
+	{
+		PlayerViewport->SetChannelMask(InMask);
+	}
+}
+
+void SMediaPlayerEditorViewer::ToggleChannelMask(MediaPlayerEditor::MediaImage::ETextureChannelMask InChannelToToggle)
+{
+	if (!PlayerViewport.IsValid())
+	{
+		return;
+	}
+
+	using namespace MediaPlayerEditor::MediaImage;
+
+	const ETextureChannelMask CurrentMask = PlayerViewport->GetChannelMask();
+
+	if (EnumHasAllFlags(CurrentMask, InChannelToToggle))
+	{
+		PlayerViewport->SetChannelMask(CurrentMask & ~InChannelToToggle);
+	}
+	else
+	{
+		PlayerViewport->SetChannelMask(CurrentMask | InChannelToToggle);
+	}
+}
+
 /* SWidget interface
  *****************************************************************************/
 
diff --git a/Engine/Plugins/Media/MediaPlayerEditor/Source/MediaPlayerEditor/Private/Widgets/SMediaPlayerEditorViewport.cpp b/Engine/Plugins/Media/MediaPlayerEditor/Source/MediaPlayerEditor/Private/Widgets/SMediaPlayerEditorViewport.cpp
--- a/Engine/Plugins/Media/MediaPlayerEditor/Source/MediaPlayerEditor/Private/Widgets/SMediaPlayerEditorViewport.cpp
+++ b/Engine/Plugins/Media/MediaPlayerEditor/Source/MediaPlayerEditor/Private/Widgets/SMediaPlayerEditorViewport.cpp
@@ -67,7 +67,7 @@ void SMediaPlayerEditorViewport::Construct(const FArguments& InArgs, UMediaPlaye
 						})
 						[
 							// movie texture
-							SNew(SMediaPlayerEditorOutput, InMediaPlayer, InMediaTexture, bInIsSoundEnabled)
+							SAssignNew(Output, SMediaPlayerEditorOutput, InMediaPlayer, InMediaTexture, bInIsSoundEnabled)
 						]
 					]
 				]
@@ -155,6 +155,24 @@ void SMediaPlayerEditorViewport::Construct(const FArguments& InArgs, UMediaPlaye
 	];
 }
 
+MediaPlayerEditor::MediaImage::ETextureChannelMask SMediaPlayerEditorViewport::GetChannelMask() const
+{
+	if (Output.IsValid())
+	{
+		return Output->GetChannelMask();
+	}
+
+	return MediaPlayerEditor::MediaImage::ETextureChannelMask::RGBA;
+}
+
+void SMediaPlayerEditorViewport::SetChannelMask(MediaPlayerEditor::MediaImage::ETextureChannelMask InMask)
+{
+	if (Output.IsValid())
+	{
+		Output->SetChannelMask(InMask);
+	}
+}
+
 
 /* SWidget interface
  *****************************************************************************/
@@ -163,7 +181,7 @@ FReply SMediaPlayerEditorViewport::OnMouseButtonDoubleClick(const FGeometry& MyG
 {
 	if (TSharedPtr<FUICommandList> Commands = CommandsWeak.Pin())
 	{
-		Commands->ExecuteAction(FMediaPlayerEditorCommands::Get().TogglePlayPauseMedia.ToSharedRef());
+		Commands->TryExecuteAction(FMediaPlayerEditorCommands::Get().TogglePlayPauseMedia.ToSharedRef());
 	}
 
 	return FReply::Handled();
diff --git a/Engine/Plugins/Media/MediaPlayerEditor/Source/MediaPlayerEditor/Private/Widgets/SMediaPlayerEditorViewport.h b/Engine/Plugins/Media/MediaPlayerEditor/Source/MediaPlayerEditor/Private/Widgets/SMediaPlayerEditorViewport.h
--- a/Engine/Plugins/Media/MediaPlayerEditor/Source/MediaPlayerEditor/Private/Widgets/SMediaPlayerEditorViewport.h
+++ b/Engine/Plugins/Media/MediaPlayerEditor/Source/MediaPlayerEditor/Private/Widgets/SMediaPlayerEditorViewport.h
@@ -5,8 +5,10 @@
 #include "Styling/ISlateStyle.h"
 #include "Widgets/DeclarativeSyntaxSupport.h"
 #include "Widgets/SCompoundWidget.h"
+#include "Widgets/SMediaImage.h"
 
 class FUICommandList;
+class SMediaPlayerEditorOutput;
 class UMediaPlayer;
 class UMediaSoundComponent;
 class UMediaTexture;
@@ -46,6 +48,20 @@ public:
 	 */
 	void EnableMouseControl(bool bIsEnabled) { bIsMouseControlEnabled = bIsEnabled; }
 
+	/**
+	 * Returns the active mask.
+	 */
+	MediaPlayerEditor::MediaImage::ETextureChannelMask GetChannelMask() const;
+
+	/**
+	 * Sets a channel-based mask for the image.
+	 *
+	 * If only a single channel is displayed, it will show in greyscale.
+	 *
+	 * @param InMask to display.
+	 */
+	void SetChannelMask(MediaPlayerEditor::MediaImage::ETextureChannelMask InMask);
+
 public:
 
 	//~ SWidget interface
@@ -81,6 +97,9 @@ private:
 	/** True if the mouse can control things. */
 	bool bIsMouseControlEnabled;
 
+	/** The output for the media player. */
+	TSharedPtr<SMediaPlayerEditorOutput> Output;
+
 	/** Command list for executing commands. */
 	TWeakPtr<FUICommandList> CommandsWeak;
 };
\ No newline at end of file
diff --git a/Engine/Plugins/Media/MediaPlayerEditor/Source/MediaPlayerEditor/Public/Models/MediaPlayerEditorCommands.h b/Engine/Plugins/Media/MediaPlayerEditor/Source/MediaPlayerEditor/Public/Models/MediaPlayerEditorCommands.h
--- a/Engine/Plugins/Media/MediaPlayerEditor/Source/MediaPlayerEditor/Public/Models/MediaPlayerEditorCommands.h
+++ b/Engine/Plugins/Media/MediaPlayerEditor/Source/MediaPlayerEditor/Public/Models/MediaPlayerEditorCommands.h
@@ -76,6 +76,18 @@ public:
 
 	/** Jump backward media 10s. */
 	TSharedPtr<FUICommandInfo> JumpBackward;
+
+	/** Toggles red texture channel. */
+	TSharedPtr<FUICommandInfo> ToggledRedTextureChannel;
+
+	/** Toggles green texture channel. */
+	TSharedPtr<FUICommandInfo> ToggledGreenTextureChannel;
+
+	/** Toggles blue texture channel. */
+	TSharedPtr<FUICommandInfo> ToggledBlueTextureChannel;
+
+	/** Toggles alpha texture channel. */
+	TSharedPtr<FUICommandInfo> ToggledAlphaTextureChannel;
 };
 
 #undef UE_API
diff --git a/Engine/Plugins/Media/MediaPlayerEditor/Source/MediaPlayerEditor/Public/Widgets/SMediaImage.h b/Engine/Plugins/Media/MediaPlayerEditor/Source/MediaPlayerEditor/Public/Widgets/SMediaImage.h
--- a/Engine/Plugins/Media/MediaPlayerEditor/Source/MediaPlayerEditor/Public/Widgets/SMediaImage.h
+++ b/Engine/Plugins/Media/MediaPlayerEditor/Source/MediaPlayerEditor/Public/Widgets/SMediaImage.h
@@ -4,15 +4,34 @@
 
 #include "UObject/GCObject.h"
 
+#include "UObject/StrongObjectPtr.h"
 #include "Widgets/DeclarativeSyntaxSupport.h"
 #include "Widgets/SCompoundWidget.h"
 
 #define UE_API MEDIAPLAYEREDITOR_API
 
 class UMaterial;
-class UMaterialExpressionTextureSample;
+class UMaterialInstanceDynamic;
 class UTexture;
 
+namespace MediaPlayerEditor
+{
+	namespace MediaImage
+	{
+		enum class ETextureChannelMask : uint8
+		{
+			None  = 0,
+			Red   = 1 << 0,
+			Green = 1 << 1,
+			Blue  = 1 << 2,
+			Alpha = 1 << 3,
+			RGB   = Red | Green | Blue,
+			RGBA  = Red | Green | Blue | Alpha,
+		};
+		ENUM_CLASS_FLAGS(ETextureChannelMask);
+	}
+}
+
 class SMediaImage : public SCompoundWidget
 {
 public:
@@ -44,46 +63,37 @@ public:
 	 */
 	UE_API virtual void Tick(const FGeometry& InAllottedGeometry, const double InCurrentTime, const float InDeltaTime) override;
 
-private:
-
-	class FInternalReferenceCollector : public FGCObject
-	{
-	public:
-		FInternalReferenceCollector(SMediaImage* InObject)
-			: Object(InObject)
-		{
-		}
-
-		//~ FGCObject interface
-		virtual void AddReferencedObjects(FReferenceCollector& InCollector) override
-		{
-			InCollector.AddReferencedObject(Object->Material);
-			InCollector.AddReferencedObject(Object->TextureSampler);
-		}
-		virtual FString GetReferencerName() const override
-		{
-			return TEXT("SMediaImage::FInternalReferenceCollector");
-		}
+	/**
+	 * Returns the active mask.
+	 */
+	UE_API MediaPlayerEditor::MediaImage::ETextureChannelMask GetChannelMask() const;
 
-	private:
-		SMediaImage* Object;
-	};
-	friend FInternalReferenceCollector;
+	/**
+	 * Sets a channel-based mask for the image.
+	 * 
+	 * If only a single channel is displayed, it will show in greyscale.
+	 * 
+	 * @param InMask to display.
+	 */
+	UE_API void SetChannelMask(MediaPlayerEditor::MediaImage::ETextureChannelMask InMask);
 
-	/** Collector to keep the UObject alive. */
-	FInternalReferenceCollector Collector;
+private:
+	static TMap<MediaPlayerEditor::MediaImage::ETextureChannelMask, TStrongObjectPtr<UMaterial>> CachedMaterials;
 
-	/** The material that wraps the video texture for display in an SImage. */
-	TObjectPtr<UMaterial> Material;
+	/** The Slate brush that renders the material and its nodes. */
+	TStrongObjectPtr<UMaterialInstanceDynamic> MaterialAInstance;
 
-	/** The Slate brush that renders the material. */
+	/** The brush used to render the media. */
 	TSharedPtr<FSlateBrush> MaterialBrush;
 
-	/** The video texture sampler in the wrapper material. */
-	TObjectPtr<UMaterialExpressionTextureSample> TextureSampler;
-
 	/** Brush image size attribute. */
 	TAttribute<FVector2D> BrushImageSize;
+
+	/** The texture being displayed */
+	TWeakObjectPtr<UTexture> TextureWeak;
+
+	/** The active mask */
+	MediaPlayerEditor::MediaImage::ETextureChannelMask Mask;
 };
 
 #undef UE_API
diff --git a/Engine/Plugins/Media/MediaPlayerEditor/Source/MediaPlayerEditor/Public/Widgets/SMediaImageTextureChannelToggle.h b/Engine/Plugins/Media/MediaPlayerEditor/Source/MediaPlayerEditor/Public/Widgets/SMediaImageTextureChannelToggle.h
new file mode 100644
--- /dev/null
+++ b/Engine/Plugins/Media/MediaPlayerEditor/Source/MediaPlayerEditor/Public/Widgets/SMediaImageTextureChannelToggle.h
@@ -0,0 +1,51 @@
+// Copyright Epic Games, Inc. All Rights Reserved.
+
+#pragma once
+
+#include "Widgets/SCompoundWidget.h"
+
+#include "Widgets/SMediaImage.h"
+
+class FUICommandList;
+class SMediaPlayerEditorViewer;
+enum class ECheckBoxState : uint8;
+
+#define UE_API MEDIAPLAYEREDITOR_API
+
+/**
+ * Implements a button to toggle the contents of a texture channel for the SMediaPlayerEditorViewer
+ */
+class SMediaImageTextureChannelToggle
+	: public SCompoundWidget
+{
+public:
+	SLATE_BEGIN_ARGS(SMediaImageTextureChannelToggle)
+		{
+		}
+	SLATE_END_ARGS()
+
+	/**
+	 * Construct this widget
+	 *
+	 * @param InArgs The declaration data for this widget.
+	 */
+	UE_API void Construct(const FArguments& InArgs, TSharedPtr<SMediaPlayerEditorViewer> InViewer,
+		MediaPlayerEditor::MediaImage::ETextureChannelMask InChannel, TSharedPtr<FUICommandList> InCommandList);
+
+private:
+	TWeakPtr<SMediaPlayerEditorViewer> ViewerWeak;
+	MediaPlayerEditor::MediaImage::ETextureChannelMask Channel;
+	TWeakPtr<FUICommandList> CommandListWeak;
+
+	FSlateColor GetChannelButtonBackgroundColor() const;
+
+	FSlateColor GetChannelButtonForegroundColor() const;
+
+	ECheckBoxState OnGetChannelButtonCheckState() const;
+
+	bool IsChannelButtonEnabled() const;
+
+	void OnChannelButtonCheckStateChanged(ECheckBoxState InNewState);
+};
+
+#undef UE_API
diff --git a/Engine/Plugins/Media/MediaPlayerEditor/Source/MediaPlayerEditor/Public/Widgets/SMediaPlayerEditorViewer.h b/Engine/Plugins/Media/MediaPlayerEditor/Source/MediaPlayerEditor/Public/Widgets/SMediaPlayerEditorViewer.h
--- a/Engine/Plugins/Media/MediaPlayerEditor/Source/MediaPlayerEditor/Public/Widgets/SMediaPlayerEditorViewer.h
+++ b/Engine/Plugins/Media/MediaPlayerEditor/Source/MediaPlayerEditor/Public/Widgets/SMediaPlayerEditorViewer.h
@@ -10,6 +10,7 @@
 #include "UObject/NameTypes.h"
 #include "Widgets/DeclarativeSyntaxSupport.h"
 #include "Widgets/SCompoundWidget.h"
+#include "Widgets/SMediaImage.h"
 
 #define UE_API MEDIAPLAYEREDITOR_API
 
@@ -67,11 +68,37 @@ public:
 		UMediaTexture* InMediaTexture, const TSharedRef<ISlateStyle>& InStyle,
 		bool bInIsSoundEnabled);
 
+	UE_API UMediaPlayer* GetMediaPlayer() const;
+
 	/**
 	 * Enables/disables using the mouse to control the viewport.
 	 */
 	UE_API void EnableMouseControl(bool bIsEnabled);
 
+	/**
+	 * Returns the active mask.
+	 */
+	UE_API MediaPlayerEditor::MediaImage::ETextureChannelMask GetChannelMask() const;
+
+	/**
+	 * Returns true if all the given channel mask is enabled.
+	 */
+	bool IsChannelMasked(MediaPlayerEditor::MediaImage::ETextureChannelMask InChannel) const;
+
+	/**
+	 * Sets a channel-based mask for the image.
+	 *
+	 * If only a single channel is displayed, it will show in greyscale.
+	 *
+	 * @param InMask to display.
+	 */
+	UE_API void SetChannelMask(MediaPlayerEditor::MediaImage::ETextureChannelMask InMask);
+
+	/**
+	 * Toggles the given mask in the channel mask.
+	 */
+	void ToggleChannelMask(MediaPlayerEditor::MediaImage::ETextureChannelMask InChannelToToggle);
+
 public:
 
 	//~ SWidget interface
