Media Player Editor: Fixed padding on texture channel toggle buttons.

#Synced-CL 50961346

diff --git a/Engine/Plugins/Media/MediaPlayerEditor/Source/MediaPlayerEditor/Private/Widgets/SMediaImageTextureChannelToggle.cpp b/Engine/Plugins/Media/MediaPlayerEditor/Source/MediaPlayerEditor/Private/Widgets/SMediaImageTextureChannelToggle.cpp
--- a/Engine/Plugins/Media/MediaPlayerEditor/Source/MediaPlayerEditor/Private/Widgets/SMediaImageTextureChannelToggle.cpp
+++ b/Engine/Plugins/Media/MediaPlayerEditor/Source/MediaPlayerEditor/Private/Widgets/SMediaImageTextureChannelToggle.cpp
@@ -7,6 +7,7 @@
 #include "Models/MediaPlayerEditorCommands.h"
 #include "Styling/AppStyle.h"
 #include "Widgets/Input/SCheckBox.h"
+#include "Widgets/Layout/SBox.h"
 #include "Widgets/Text/STextBlock.h"
 #include "Widgets/SMediaPlayerEditorViewer.h"
 
@@ -47,19 +48,22 @@ void SMediaImageTextureChannelToggle::Construct(const FArguments& InArgs, TShare
 
 	ChildSlot
 	[
-		SNew(SCheckBox)
-		.Style(FAppStyle::Get(), "TextureEditor.ChannelButtonStyle")
-		.BorderBackgroundColor(this, &SMediaImageTextureChannelToggle::GetChannelButtonBackgroundColor)
-		.ForegroundColor(this, &SMediaImageTextureChannelToggle::GetChannelButtonForegroundColor)
-		.OnCheckStateChanged(this, &SMediaImageTextureChannelToggle::OnChannelButtonCheckStateChanged)
-		.IsChecked(this, &SMediaImageTextureChannelToggle::OnGetChannelButtonCheckState)
-		.IsEnabled(this, &SMediaImageTextureChannelToggle::IsChannelButtonEnabled)
-		.ToolTipText(Command->GetDescription())
+		SNew(SBox)
 		.Padding(1.f)
 		[
-			SNew(STextBlock)
-			.Font(FAppStyle::Get().GetFontStyle("TextureEditor.ChannelButtonFont"))
-			.Text(Command->GetLabel())
+			SNew(SCheckBox)
+			.Style(FAppStyle::Get(), "TextureEditor.ChannelButtonStyle")
+			.BorderBackgroundColor(this, &SMediaImageTextureChannelToggle::GetChannelButtonBackgroundColor)
+			.ForegroundColor(this, &SMediaImageTextureChannelToggle::GetChannelButtonForegroundColor)
+			.OnCheckStateChanged(this, &SMediaImageTextureChannelToggle::OnChannelButtonCheckStateChanged)
+			.IsChecked(this, &SMediaImageTextureChannelToggle::OnGetChannelButtonCheckState)
+			.IsEnabled(this, &SMediaImageTextureChannelToggle::IsChannelButtonEnabled)
+			.ToolTipText(Command->GetDescription())
+			[
+				SNew(STextBlock)
+				.Font(FAppStyle::Get().GetFontStyle("TextureEditor.ChannelButtonFont"))
+				.Text(Command->GetLabel())
+			]
 		]
 	];
 }
