Media Player Editor: Moved the toolbar buttons to be near the scrub bar. Added new buttons. Made buttons smaller.

#Synced-CL 51022704

diff --git a/Engine/Plugins/Media/MediaPlate/Source/MediaPlate/Private/MediaPlateComponent.cpp b/Engine/Plugins/Media/MediaPlate/Source/MediaPlate/Private/MediaPlateComponent.cpp
--- a/Engine/Plugins/Media/MediaPlate/Source/MediaPlate/Private/MediaPlateComponent.cpp
+++ b/Engine/Plugins/Media/MediaPlate/Source/MediaPlate/Private/MediaPlateComponent.cpp
@@ -366,7 +366,12 @@ void UMediaPlateComponent::TickComponent(float DeltaTime, enum ELevelTick TickTy
 				PendingPlaybackState = IntendedPlaybackState;
 				if (IntendedPlaybackState == EPlaybackState::Playing)
 				{
-					MediaPlayer->Play();
+					if (FMath::IsNearlyZero(CurrentRate) || !MediaPlayer->SupportsRate(CurrentRate, /* Unthinned */ true))
+					{
+						CurrentRate = 1.f;
+					}
+
+					MediaPlayer->SetRate(CurrentRate);
 				}
 				else
 				{
@@ -443,12 +448,7 @@ bool UMediaPlateComponent::IsEventStateChangeAllowed(EMediaPlateEventState InEve
 	{
 	case EMediaPlateEventState::Open:
 	case EMediaPlateEventState::Play:
-	case EMediaPlateEventState::JumpForward:
-	case EMediaPlateEventState::JumpBackward:
 		return !IsAnyMediaTextureProxied();
-	case EMediaPlateEventState::StepForward:
-	case EMediaPlateEventState::StepBackward:
-		return IntendedPlaybackState == EPlaybackState::Paused && !IsAnyMediaTextureProxied();
 	default:
 		return true;
 	}
@@ -544,6 +544,12 @@ void UMediaPlateComponent::Play()
 	CurrentRate = 1.0f;
 }
 
+void UMediaPlateComponent::PlayReverse()
+{
+	IntendedPlaybackState = EPlaybackState::Playing;
+	CurrentRate = -1.f;
+}
+
 void UMediaPlateComponent::Pause()
 {
 	IntendedPlaybackState = EPlaybackState::Paused;
@@ -1850,6 +1856,11 @@ void UMediaPlateComponent::SwitchStates(EMediaPlateEventState State)
 			Play();
 		}
 		break;
+	case EMediaPlateEventState::PlayReverse:
+		{
+			PlayReverse();
+		}
+		break;
 	case EMediaPlateEventState::Open:
 		{
 			Open();
@@ -1899,7 +1910,7 @@ void UMediaPlateComponent::SwitchStates(EMediaPlateEventState State)
 			Previous();
 		}
 		break;
-	case EMediaPlateEventState::FastForward:
+	case EMediaPlateEventState::JumpToEnd:
 		{
 			if (MediaPlayer != nullptr)
 			{
@@ -1935,41 +1946,6 @@ void UMediaPlateComponent::SwitchStates(EMediaPlateEventState State)
 			}
 		}
 		break;
-	case EMediaPlateEventState::JumpForward:
-		{
-			if (MediaPlayer != nullptr)
-			{
-				const FTimespan CurrentTime = MediaPlayer->GetTime();
-				const FTimespan Duration = MediaPlayer->GetDuration();
-
-				if ((Duration - CurrentTime) <= TenSeconds)
-				{
-					MediaPlayer->Seek(Duration);
-				}
-				else
-				{
-					MediaPlayer->Seek(MediaPlayer->GetTime() + FTimespan::FromSeconds(10.f));
-				}
-			}
-		}
-		break;
-	case EMediaPlateEventState::JumpBackward:
-		{
-			if (MediaPlayer != nullptr)
-			{
-				const FTimespan CurrentTime = MediaPlayer->GetTime();
-
-				if (CurrentTime <= TenSeconds)
-				{
-					MediaPlayer->Seek(0);
-				}
-				else
-				{
-					MediaPlayer->Seek(MediaPlayer->GetTime() - FTimespan::FromSeconds(10.f));
-				}
-			}
-		}
-		break;
 	case EMediaPlateEventState::MAX:
 	default:
 		checkNoEntry();
diff --git a/Engine/Plugins/Media/MediaPlate/Source/MediaPlate/Public/MediaPlateComponent.h b/Engine/Plugins/Media/MediaPlate/Source/MediaPlate/Public/MediaPlateComponent.h
--- a/Engine/Plugins/Media/MediaPlate/Source/MediaPlate/Public/MediaPlateComponent.h
+++ b/Engine/Plugins/Media/MediaPlate/Source/MediaPlate/Public/MediaPlateComponent.h
@@ -38,11 +38,10 @@ enum class EMediaPlateEventState : uint8
 	Rewind,
 	Next,
 	Previous,
-	FastForward,
+	PlayReverse,
+	JumpToEnd,
 	StepForward,
 	StepBackward,
-	JumpForward,
-	JumpBackward,
 	MAX
 };
 
@@ -148,12 +147,19 @@ public:
 	UE_API bool Next();
 
 	/**
-	 * Call this to start playing.
+	 * Call this to start playing at a rate of 1.0 (forward at normal speed)
 	 * Open must be called before this.
 	 */
 	UFUNCTION(BlueprintCallable, Category = "Media|MediaPlateComponent")
 	UE_API void Play();
 
+	/**
+	 * Call this to start playing at a rate of -1.0 (reverse at normal speed)
+	 * Open must be called before this.
+	 */
+	UFUNCTION(BlueprintCallable, Category = "Media|MediaPlateComponent")
+	UE_API void PlayReverse();
+
 	/**
 	 * Call this to pause playback.
 	 * Play can be called to resume playback.
diff --git a/Engine/Plugins/Media/MediaPlate/Source/MediaPlateEditor/Private/MediaPlateCustomization.cpp b/Engine/Plugins/Media/MediaPlate/Source/MediaPlateEditor/Private/MediaPlateCustomization.cpp
--- a/Engine/Plugins/Media/MediaPlate/Source/MediaPlateEditor/Private/MediaPlateCustomization.cpp
+++ b/Engine/Plugins/Media/MediaPlate/Source/MediaPlateEditor/Private/MediaPlateCustomization.cpp
@@ -1079,6 +1079,11 @@ bool FMediaPlateCustomization::IsMediaPlateEventAllowedForPlayer(EMediaPlateEven
 		return InMediaPlayer->IsReady()
 			&& (!InMediaPlayer->IsPlaying() || InMediaPlayer->GetRate() != 1.0f);
 
+	case EMediaPlateEventState::PlayReverse:
+		// Is player paused or fast forwarding/rewinding?
+		return InMediaPlayer->IsReady() && InMediaPlayer->SupportsRate(-1.f, /* Unthinned */ true)
+			&& (!InMediaPlayer->IsPlaying() || InMediaPlayer->GetRate() != -1.0f);
+
 	case EMediaPlateEventState::Open:
 		return true; // The condition is implemented by the media plate already.
 
@@ -1097,6 +1102,13 @@ bool FMediaPlateCustomization::IsMediaPlateEventAllowedForPlayer(EMediaPlateEven
 	case EMediaPlateEventState::Rewind:
 		return InMediaPlayer->IsReady() && InMediaPlayer->SupportsSeeking() && InMediaPlayer->GetTime() > FTimespan::Zero();
 
+	case EMediaPlateEventState::JumpToEnd:
+		return InMediaPlayer->IsReady() && InMediaPlayer->SupportsSeeking() && InMediaPlayer->GetTime() < InMediaPlayer->GetDuration();
+
+	case EMediaPlateEventState::StepForward:
+	case EMediaPlateEventState::StepBackward:
+		return InMediaPlayer->IsReady() && InMediaPlayer->SupportsSeeking() && InMediaPlayer->IsPaused();
+	
 	//case EMediaPlateEventState::Next:	// was not implemented
 	//case EMediaPlateEventState::Previous: // was not implemented
 
diff --git a/Engine/Plugins/Media/MediaPlate/Source/MediaPlateEditor/Private/Toolkits/MediaPlateEditorToolkit.cpp b/Engine/Plugins/Media/MediaPlate/Source/MediaPlateEditor/Private/Toolkits/MediaPlateEditorToolkit.cpp
--- a/Engine/Plugins/Media/MediaPlate/Source/MediaPlateEditor/Private/Toolkits/MediaPlateEditorToolkit.cpp
+++ b/Engine/Plugins/Media/MediaPlate/Source/MediaPlateEditor/Private/Toolkits/MediaPlateEditorToolkit.cpp
@@ -320,6 +320,12 @@ void FMediaPlateEditorToolkit::BindCommands()
 		FCanExecuteAction::CreateSP( this, &FMediaPlateEditorToolkit::CanHandleMediaPlateEvent, EMediaPlateEventState::Play )
 	);
 
+	ToolkitCommands->MapAction(
+		Commands.PlayReverseMedia,
+		FExecuteAction::CreateSP(this, &FMediaPlateEditorToolkit::HandleMediaPlateEvent, EMediaPlateEventState::PlayReverse),
+		FCanExecuteAction::CreateSP(this, &FMediaPlateEditorToolkit::CanHandleMediaPlateEvent, EMediaPlateEventState::PlayReverse)
+	);
+
 	ToolkitCommands->MapAction(
 		Commands.PreviousMedia,
 		FExecuteAction::CreateSP( this, &FMediaPlateEditorToolkit::HandleMediaPlateEvent, EMediaPlateEventState::Previous ),
@@ -338,6 +344,12 @@ void FMediaPlateEditorToolkit::BindCommands()
 		FCanExecuteAction::CreateSP(this, &FMediaPlateEditorToolkit::CanHandleMediaPlateEvent, EMediaPlateEventState::Rewind)
 	);
 
+	ToolkitCommands->MapAction(
+		Commands.JumpToEndMedia,
+		FExecuteAction::CreateSP(this, &FMediaPlateEditorToolkit::HandleMediaPlateEvent, EMediaPlateEventState::JumpToEnd),
+		FCanExecuteAction::CreateSP(this, &FMediaPlateEditorToolkit::CanHandleMediaPlateEvent, EMediaPlateEventState::JumpToEnd)
+	);
+
 	ToolkitCommands->MapAction(
 		Commands.TogglePlayPauseMedia,
 		FExecuteAction::CreateSP( this, &FMediaPlateEditorToolkit::HandleMediaPlateTogglePlayPauseEvent ),
@@ -345,35 +357,23 @@ void FMediaPlateEditorToolkit::BindCommands()
 	);
 
 	ToolkitCommands->MapAction(
-		Commands.FastForwardMedia,
-		FExecuteAction::CreateSP(this, &FMediaPlateEditorToolkit::HandleMediaPlateEvent, EMediaPlateEventState::FastForward),
-		FCanExecuteAction::CreateSP(this, &FMediaPlateEditorToolkit::CanHandleMediaPlateEvent, EMediaPlateEventState::FastForward)
+		Commands.TogglePlayReversePauseMedia,
+		FExecuteAction::CreateSP(this, &FMediaPlateEditorToolkit::HandleMediaPlateTogglePlayReversePauseEvent),
+		FCanExecuteAction::CreateSP(this, &FMediaPlateEditorToolkit::CanHandleMediaPlateTogglePlayReversePauseEvent)
 	);
 
 	ToolkitCommands->MapAction(
-		Commands.StepForward,
+		Commands.StepForwardMedia,
 		FExecuteAction::CreateSP(this, &FMediaPlateEditorToolkit::HandleMediaPlateEvent, EMediaPlateEventState::StepForward),
 		FCanExecuteAction::CreateSP(this, &FMediaPlateEditorToolkit::CanHandleMediaPlateEvent, EMediaPlateEventState::StepForward)
 	);
 
 	ToolkitCommands->MapAction(
-		Commands.StepBackward,
+		Commands.StepBackwardMedia,
 		FExecuteAction::CreateSP(this, &FMediaPlateEditorToolkit::HandleMediaPlateEvent, EMediaPlateEventState::StepBackward),
 		FCanExecuteAction::CreateSP(this, &FMediaPlateEditorToolkit::CanHandleMediaPlateEvent, EMediaPlateEventState::StepBackward)
 	);
 
-	ToolkitCommands->MapAction(
-		Commands.JumpForward,
-		FExecuteAction::CreateSP(this, &FMediaPlateEditorToolkit::HandleMediaPlateEvent, EMediaPlateEventState::JumpForward),
-		FCanExecuteAction::CreateSP(this, &FMediaPlateEditorToolkit::CanHandleMediaPlateEvent, EMediaPlateEventState::JumpForward)
-	);
-
-	ToolkitCommands->MapAction(
-		Commands.JumpBackward,
-		FExecuteAction::CreateSP(this, &FMediaPlateEditorToolkit::HandleMediaPlateEvent, EMediaPlateEventState::JumpBackward),
-		FCanExecuteAction::CreateSP(this, &FMediaPlateEditorToolkit::CanHandleMediaPlateEvent, EMediaPlateEventState::JumpBackward)
-	);
-
 	ToolkitCommands->MapAction(
 		Commands.ToggledRedTextureChannel,
 		FExecuteAction::CreateLambda([this] { ToggleChannelMask(MediaPlayerEditor::MediaImage::ETextureChannelMask::Red); }),
@@ -415,34 +415,13 @@ void FMediaPlateEditorToolkit::ExtendToolBar()
 		FToolBarExtensionDelegate::CreateLambda(
 			[this](FToolBarBuilder& ToolbarBuilder)
 			{
-				const FMediaPlayerEditorCommands& Commands = FMediaPlayerEditorCommands::GetExternal();
-
-				ToolbarBuilder.BeginSection("PlaybackControls");
-				{
-					ToolbarBuilder.AddToolBarButton(Commands.PreviousMedia);
-					ToolbarBuilder.AddToolBarButton(Commands.RewindMedia);
-					ToolbarBuilder.AddToolBarButton(Commands.ReverseMedia);
-					ToolbarBuilder.AddToolBarButton(Commands.PlayMedia);
-					ToolbarBuilder.AddToolBarButton(Commands.PauseMedia);
-					ToolbarBuilder.AddToolBarButton(Commands.ForwardMedia);
-					ToolbarBuilder.AddToolBarButton(Commands.NextMedia);
-				}
-				ToolbarBuilder.EndSection();
-
-				ToolbarBuilder.BeginSection("MediaControls");
-				{
-					ToolbarBuilder.AddToolBarButton(Commands.OpenMedia);
-					ToolbarBuilder.AddToolBarButton(Commands.CloseMedia);
-				}
-				ToolbarBuilder.EndSection();
-
 				using namespace MediaPlayerEditor::MediaImage;
 
 				ToolbarBuilder.BeginSection("TextureControls");
 				{
-					ToolbarBuilder.AddWidget(SNew(SMediaImageTextureChannelToggle, Viewer, ETextureChannelMask::Red, ToolkitCommands));
+					ToolbarBuilder.AddWidget(SNew(SMediaImageTextureChannelToggle, Viewer, ETextureChannelMask::Red,   ToolkitCommands));
 					ToolbarBuilder.AddWidget(SNew(SMediaImageTextureChannelToggle, Viewer, ETextureChannelMask::Green, ToolkitCommands));
-					ToolbarBuilder.AddWidget(SNew(SMediaImageTextureChannelToggle, Viewer, ETextureChannelMask::Blue, ToolkitCommands));
+					ToolbarBuilder.AddWidget(SNew(SMediaImageTextureChannelToggle, Viewer, ETextureChannelMask::Blue,  ToolkitCommands));
 					ToolbarBuilder.AddWidget(SNew(SMediaImageTextureChannelToggle, Viewer, ETextureChannelMask::Alpha, ToolkitCommands));
 				}
 				ToolbarBuilder.EndSection();
@@ -513,14 +492,16 @@ bool FMediaPlateEditorToolkit::CanHandleMediaPlateTogglePlayPauseEvent()
 {
 	if (UMediaPlayer* MediaPlayer = MediaPlate->GetMediaPlayer())
 	{
-		if (MediaPlayer->IsPlaying())
+		if (MediaPlayer->IsReady())
 		{
-			return CanHandleMediaPlateEvent(EMediaPlateEventState::Pause);
-		}
-
-		if (MediaPlayer->IsPaused())
-		{
-			return CanHandleMediaPlateEvent(EMediaPlateEventState::Play);
+			if (MediaPlayer->IsPlaying() && MediaPlayer->GetRate() > 0)
+			{
+				return CanHandleMediaPlateEvent(EMediaPlateEventState::Pause);
+			}
+			else
+			{
+				return CanHandleMediaPlateEvent(EMediaPlateEventState::Play);
+			}
 		}
 	}
 
@@ -531,13 +512,54 @@ void FMediaPlateEditorToolkit::HandleMediaPlateTogglePlayPauseEvent()
 {
 	if (UMediaPlayer* MediaPlayer = MediaPlate->GetMediaPlayer())
 	{
-		if (MediaPlayer->IsPlaying())
+		if (MediaPlayer->IsReady())
 		{
-			HandleMediaPlateEvent(EMediaPlateEventState::Pause);
+			if (MediaPlayer->IsPlaying() && MediaPlayer->GetRate() > 0)
+			{
+				HandleMediaPlateEvent(EMediaPlateEventState::Pause);
+			}
+			else
+			{
+				HandleMediaPlateEvent(EMediaPlateEventState::Play);
+			}
+		}
+	}
+}
+
+bool FMediaPlateEditorToolkit::CanHandleMediaPlateTogglePlayReversePauseEvent()
+{
+	if (UMediaPlayer* MediaPlayer = MediaPlate->GetMediaPlayer())
+	{
+		if (MediaPlayer->IsReady() && MediaPlayer->SupportsRate(-1.f, /* Unthinned */ true))
+		{
+			if (MediaPlayer->IsPlaying() && MediaPlayer->GetRate() < 0)
+			{
+				return CanHandleMediaPlateEvent(EMediaPlateEventState::Pause);
+			}
+			else
+			{
+				return CanHandleMediaPlateEvent(EMediaPlateEventState::PlayReverse);
+			}
 		}
-		else if (MediaPlayer->IsPaused())
+	}
+
+	return false;
+}
+
+void FMediaPlateEditorToolkit::HandleMediaPlateTogglePlayReversePauseEvent()
+{
+	if (UMediaPlayer* MediaPlayer = MediaPlate->GetMediaPlayer())
+	{
+		if (MediaPlayer->IsReady() && MediaPlayer->SupportsRate(-1.f, /* Unthinned */ true))
 		{
-			HandleMediaPlateEvent(EMediaPlateEventState::Play);
+			if (MediaPlayer->IsPlaying() && MediaPlayer->GetRate() < 0)
+			{
+				HandleMediaPlateEvent(EMediaPlateEventState::Pause);
+			}
+			else
+			{
+				HandleMediaPlateEvent(EMediaPlateEventState::PlayReverse);
+			}
 		}
 	}
 }
diff --git a/Engine/Plugins/Media/MediaPlate/Source/MediaPlateEditor/Private/Toolkits/MediaPlateEditorToolkit.h b/Engine/Plugins/Media/MediaPlate/Source/MediaPlateEditor/Private/Toolkits/MediaPlateEditorToolkit.h
--- a/Engine/Plugins/Media/MediaPlate/Source/MediaPlateEditor/Private/Toolkits/MediaPlateEditorToolkit.h
+++ b/Engine/Plugins/Media/MediaPlate/Source/MediaPlateEditor/Private/Toolkits/MediaPlateEditorToolkit.h
@@ -118,6 +118,12 @@ private:
 	/** Callback specifically for toggle play/pause which then triggers a play or pause event. */
 	void HandleMediaPlateTogglePlayPauseEvent();
 
+	/** Callback specifically to validate the toggle between reverse play and pause. */
+	bool CanHandleMediaPlateTogglePlayReversePauseEvent();
+
+	/** Callback specifically for toggle play/pause which then triggers a reverse play or pause event. */
+	void HandleMediaPlateTogglePlayReversePauseEvent();
+
 	/** Callback for spawning tabs. */
 	TSharedRef<SDockTab> HandleTabManagerSpawnTab(const FSpawnTabArgs& Args, FName TabIdentifier);
 	
diff --git a/Engine/Plugins/Media/MediaPlayerEditor/Source/MediaPlayerEditor/Private/Models/MediaPlayerEditorCommands.cpp b/Engine/Plugins/Media/MediaPlayerEditor/Source/MediaPlayerEditor/Private/Models/MediaPlayerEditorCommands.cpp
--- a/Engine/Plugins/Media/MediaPlayerEditor/Source/MediaPlayerEditor/Private/Models/MediaPlayerEditorCommands.cpp
+++ b/Engine/Plugins/Media/MediaPlayerEditor/Source/MediaPlayerEditor/Private/Models/MediaPlayerEditorCommands.cpp
@@ -22,9 +22,9 @@ void FMediaPlayerEditorCommands::RegisterCommands()
 	UI_COMMAND(OpenMedia, "Open", "Open the current media", EUserInterfaceActionType::Button, FInputChord());
 	UI_COMMAND(PauseMedia, "Pause", "Pause media playback", EUserInterfaceActionType::Button, FInputChord());
 	UI_COMMAND(PlayMedia, "Play", "Start media playback", EUserInterfaceActionType::Button, FInputChord());
-	UI_COMMAND(NextMedia, "Next", "Jump to next item in the play list", EUserInterfaceActionType::Button, FInputChord());
-	UI_COMMAND(PreviousMedia, "Prev", "Jump to previous item in the play list", EUserInterfaceActionType::Button, FInputChord());
-	UI_COMMAND(ReverseMedia, "Reverse", "Reverse media playback", EUserInterfaceActionType::Button, FInputChord());
+	UI_COMMAND(PlayReverseMedia, "Play in Reverse", "Start media playback in reverse. Not widely supported.", EUserInterfaceActionType::Button, FInputChord());
+	UI_COMMAND(NextMedia, "Next", "Jump to next item in the play list", EUserInterfaceActionType::Button, FInputChord(EKeys::L, EModifierKey::Shift));
+	UI_COMMAND(PreviousMedia, "Prev", "Jump to previous item in the play list", EUserInterfaceActionType::Button, FInputChord(EKeys::J, EModifierKey::Shift));
 	
 	UI_COMMAND(
 		RewindMedia,
@@ -35,19 +35,19 @@ void FMediaPlayerEditorCommands::RegisterCommands()
 	);
 
 	UI_COMMAND(
-		JumpBackward,
-		"Jump Backward",
-		"Jump the media backward 10 seconds",
-		EUserInterfaceActionType::Button,
-		FInputChord(EKeys::J)
+		ReverseMedia, 
+		"Reverse", 
+		"Play media in reverse at ever increasing speed. Not widely supported.", 
+		EUserInterfaceActionType::Button, 
+		FInputChord(EKeys::J, EModifierKey::Alt)
 	);
 
 	UI_COMMAND(
-		StepBackward,
+		StepBackwardMedia,
 		"Step Backward",
 		"Step the media backward 1 frame if it is paused",
 		EUserInterfaceActionType::Button,
-		FInputChord(EKeys::J, EModifierKey::Alt)
+		FInputChord(EKeys::J)
 	);
 
 	UI_COMMAND(
@@ -55,38 +55,39 @@ void FMediaPlayerEditorCommands::RegisterCommands()
 		"Toggle Play/Pause",
 		"Toggles between media play and pause",
 		EUserInterfaceActionType::ToggleButton,
-		FInputChord(EKeys::SpaceBar),
-		FInputChord(EKeys::K)
+		FInputChord(EKeys::K),
+		FInputChord(EKeys::SpaceBar)
 	);
 
 	UI_COMMAND(
-		ForwardMedia, 
-		"Forward", 
-		"Fast forward media playback", 
-		EUserInterfaceActionType::Button, 
-		FInputChord()
+		TogglePlayReversePauseMedia,
+		"Toggle Reverse Play/Pause",
+		"Toggles between media reverse play and pause. Not widely supported.",
+		EUserInterfaceActionType::ToggleButton,
+		FInputChord(EKeys::K, EModifierKey::Control),
+		FInputChord(EKeys::SpaceBar, EModifierKey::Control)
 	);
 
 	UI_COMMAND(
-		StepForward,
-		"Step Forward",
-		"Step the media forward 1 frame if it is paused",
-		EUserInterfaceActionType::Button,
+		ForwardMedia, 
+		"Fast Forward", 
+		"Play media forward at ever increasing speed", 
+		EUserInterfaceActionType::Button, 
 		FInputChord(EKeys::L, EModifierKey::Alt)
 	);
 
 	UI_COMMAND(
-		JumpForward,
-		"Jump Forward",
-		"Jump the media forward 10 seconds",
+		StepForwardMedia,
+		"Step Forward",
+		"Step the media forward 1 frame if it is paused",
 		EUserInterfaceActionType::Button,
 		FInputChord(EKeys::L)
 	);
 
 	UI_COMMAND(
-		FastForwardMedia,
-		"Fast Forward", 
-		"Fast forwards the media to the end", 
+		JumpToEndMedia,
+		"Jumo to End", 
+		"Jump to the end of the media. Will jump back to the start if looping is enabled.", 
 		EUserInterfaceActionType::Button, 
 		FInputChord(EKeys::L, EModifierKey::Control)
 	);
diff --git a/Engine/Plugins/Media/MediaPlayerEditor/Source/MediaPlayerEditor/Private/Shared/MediaPlayerEditorStyle.cpp b/Engine/Plugins/Media/MediaPlayerEditor/Source/MediaPlayerEditor/Private/Shared/MediaPlayerEditorStyle.cpp
--- a/Engine/Plugins/Media/MediaPlayerEditor/Source/MediaPlayerEditor/Private/Shared/MediaPlayerEditorStyle.cpp
+++ b/Engine/Plugins/Media/MediaPlayerEditor/Source/MediaPlayerEditor/Private/Shared/MediaPlayerEditorStyle.cpp
@@ -34,6 +34,14 @@ FMediaPlayerEditorStyle::FMediaPlayerEditorStyle()
 	Set("MediaPlayerEditor.ReloadButton", new IMAGE_BRUSH("btn_reload_12x", Icon12x12));
 	Set("MediaPlayerEditor.SettingsButton", new IMAGE_BRUSH("btn_settings_16x", Icon12x12));
 
+	Set<FButtonStyle>("MediaPlayerEditor.MediaControlButton", FButtonStyle()
+		.SetNormal(FSlateNoResource())
+		.SetDisabled(FSlateNoResource())
+		.SetHovered(FSlateRoundedBoxBrush(FLinearColor(0.2, 0.2, 0.2, 0.5), 3.f, Icon20x20))
+		.SetPressed(FSlateRoundedBoxBrush(FLinearColor(0.1, 0.1, 0.1, 0.5), 3.f, Icon20x20))
+		.SetNormalPadding(FMargin(2.f, 2.f, 2.f, 2.f))
+		.SetPressedPadding(FMargin(2.f, 2.f, 2.f, 2.f)));
+
 	// misc
 	Set("MediaPlayerEditor.DragDropBorder", new BOX_BRUSH("border_dragdrop", 0.5f));
 	Set("MediaPlayerEditor.MediaSourceOpened", new IMAGE_BRUSH("mediasource_opened", Icon8x8));
@@ -46,26 +54,20 @@ FMediaPlayerEditorStyle::FMediaPlayerEditorStyle()
 	Set("MediaPlayerEditor.Tabs.Stats", new IMAGE_BRUSH("tab_stats_16x", Icon16x16));
 
 	// toolbar icons
-	Set("MediaPlayerEditor.CloseMedia", new IMAGE_BRUSH("icon_eject_40x", Icon40x40));
-	Set("MediaPlayerEditor.ForwardMedia", new IMAGE_BRUSH("icon_forward_40x", Icon40x40));
-	Set("MediaPlayerEditor.ForwardMedia.Small", new IMAGE_BRUSH("icon_forward_40x", Icon20x20));
-	Set("MediaPlayerEditor.GenerateThumbnail", new CORE_IMAGE_BRUSH_SVG("Starship/AssetEditors/SaveThumbnail", Icon40x40));
-	Set("MediaPlayerEditor.NextMedia", new IMAGE_BRUSH("icon_step_40x", Icon40x40));
-	Set("MediaPlayerEditor.NextMedia.Small", new IMAGE_BRUSH("icon_step_40x", Icon20x20));
-	Set("MediaPlayerEditor.OpenMedia", new IMAGE_BRUSH("icon_open_40x", Icon40x40));
-	Set("MediaPlayerEditor.OpenMedia.Small", new IMAGE_BRUSH("icon_open_40x", Icon20x20));
-	Set("MediaPlayerEditor.PauseMedia", new IMAGE_BRUSH("icon_pause_40x", Icon40x40));
-	Set("MediaPlayerEditor.PauseMedia.Small", new IMAGE_BRUSH("icon_pause_40x", Icon20x20));
-	Set("MediaPlayerEditor.PlayMedia", new IMAGE_BRUSH("icon_play_40x", Icon40x40));
-	Set("MediaPlayerEditor.PlayMedia.Small", new IMAGE_BRUSH("icon_play_40x", Icon20x20));
-	Set("MediaPlayerEditor.PreviousMedia", new IMAGE_BRUSH("icon_step_back_40x", Icon40x40));
-	Set("MediaPlayerEditor.PreviousMedia.Small", new IMAGE_BRUSH("icon_step_back_40x", Icon20x20));
-	Set("MediaPlayerEditor.ReverseMedia", new IMAGE_BRUSH("icon_reverse_40x", Icon40x40));
-	Set("MediaPlayerEditor.ReverseMedia.Small", new IMAGE_BRUSH("icon_reverse_40x", Icon20x20));
-	Set("MediaPlayerEditor.RewindMedia", new IMAGE_BRUSH("icon_rewind_40x", Icon40x40));
-	Set("MediaPlayerEditor.RewindMedia.Small", new IMAGE_BRUSH("icon_rewind_40x", Icon20x20));
-	Set("MediaPlayerEditor.StopMedia", new IMAGE_BRUSH("icon_stop_40x", Icon40x40));
-	Set("MediaPlayerEditor.StopMedia.Small", new IMAGE_BRUSH("icon_stop_40x", Icon20x20));
+	Set("MediaPlayerEditor.GenerateThumbnail", new CORE_IMAGE_BRUSH_SVG("Starship/AssetEditors/SaveThumbnail", Icon20x20));
+	//Set("MediaPlayerEditor.ForwardMedia", new CORE_IMAGE_BRUSH_SVG("Sequencer/PlaybackControls/SaveThumbnail", Icon20x20));
+	Set("MediaPlayerEditor.NextMedia", new CORE_IMAGE_BRUSH_SVG("Starship/Common/NextArrow", Icon20x20));
+	Set("MediaPlayerEditor.PauseMedia", new CORE_IMAGE_BRUSH_SVG("Sequencer/PlaybackControls/PlayControlsPause", Icon20x20));
+	Set("MediaPlayerEditor.PlayMedia", new CORE_IMAGE_BRUSH_SVG("Sequencer/PlaybackControls/PlayControlsPlayForward", Icon20x20));
+	Set("MediaPlayerEditor.PlayReverseMedia", new CORE_IMAGE_BRUSH_SVG("Sequencer/PlaybackControls/PlayControlsPlayReverse", Icon20x20));
+	Set("MediaPlayerEditor.PreviousMedia", new CORE_IMAGE_BRUSH_SVG("Starship/Common/PreviousArrow", Icon20x20));
+	//Set("MediaPlayerEditor.ReverseMedia", new CORE_IMAGE_BRUSH_SVG("Sequencer/PlaybackControls/SaveThumbnail", Icon20x20));
+	Set("MediaPlayerEditor.RewindMedia", new CORE_IMAGE_BRUSH_SVG("Sequencer/PlaybackControls/PlayControlsToFront", Icon20x20));
+	Set("MediaPlayerEditor.JumpToEndMedia", new CORE_IMAGE_BRUSH_SVG("Sequencer/PlaybackControls/PlayControlsToEnd", Icon20x20));
+	Set("MediaPlayerEditor.StepForwardMedia", new CORE_IMAGE_BRUSH_SVG("Sequencer/PlaybackControls/PlayControlsToNext", Icon20x20));
+	Set("MediaPlayerEditor.StepBackwardMedia", new CORE_IMAGE_BRUSH_SVG("Sequencer/PlaybackControls/PlayControlsToPrevious", Icon20x20));	
+	Set("MediaPlayerEditor.CloseMedia", new CORE_IMAGE_BRUSH_SVG("Starship/MainToolbar/eject", Icon20x20));
+	Set("MediaPlayerEditor.OpenMedia", new IMAGE_BRUSH_SVG("icon_open_20x", Icon20x20));
 
 	// scrubber
 	Set("MediaPlayerEditor.Scrubber", FSliderStyle()
diff --git a/Engine/Plugins/Media/MediaPlayerEditor/Source/MediaPlayerEditor/Private/Toolkits/MediaPlayerEditorToolkit.cpp b/Engine/Plugins/Media/MediaPlayerEditor/Source/MediaPlayerEditor/Private/Toolkits/MediaPlayerEditorToolkit.cpp
--- a/Engine/Plugins/Media/MediaPlayerEditor/Source/MediaPlayerEditor/Private/Toolkits/MediaPlayerEditorToolkit.cpp
+++ b/Engine/Plugins/Media/MediaPlayerEditor/Source/MediaPlayerEditor/Private/Toolkits/MediaPlayerEditorToolkit.cpp
@@ -190,26 +190,6 @@ void FMediaPlayerEditorToolkit::ExtendToolBar()
 		FToolBarExtensionDelegate::CreateLambda(
 			[this](FToolBarBuilder& ToolbarBuilder)
 			{
-				const FMediaPlayerEditorCommands& Commands = FMediaPlayerEditorCommands::Get();
-
-				ToolbarBuilder.BeginSection("PlaybackControls");
-				{
-					ToolbarBuilder.AddToolBarButton(Commands.PreviousMedia);
-					ToolbarBuilder.AddToolBarButton(Commands.RewindMedia);
-					ToolbarBuilder.AddToolBarButton(Commands.ReverseMedia);
-					ToolbarBuilder.AddToolBarButton(Commands.PlayMedia);
-					ToolbarBuilder.AddToolBarButton(Commands.PauseMedia);
-					ToolbarBuilder.AddToolBarButton(Commands.ForwardMedia);
-					ToolbarBuilder.AddToolBarButton(Commands.NextMedia);
-				}
-				ToolbarBuilder.EndSection();
-
-				ToolbarBuilder.BeginSection("MediaControls");
-				{
-					ToolbarBuilder.AddToolBarButton(Commands.CloseMedia);
-				}
-				ToolbarBuilder.EndSection();
-
 				using namespace MediaPlayerEditor::MediaImage;
 
 				ToolbarBuilder.BeginSection("TextureControls");
diff --git a/Engine/Plugins/Media/MediaPlayerEditor/Source/MediaPlayerEditor/Private/Toolkits/MediaPlayerEditorToolkitMediaPlayerBase.cpp b/Engine/Plugins/Media/MediaPlayerEditor/Source/MediaPlayerEditor/Private/Toolkits/MediaPlayerEditorToolkitMediaPlayerBase.cpp
--- a/Engine/Plugins/Media/MediaPlayerEditor/Source/MediaPlayerEditor/Private/Toolkits/MediaPlayerEditorToolkitMediaPlayerBase.cpp
+++ b/Engine/Plugins/Media/MediaPlayerEditor/Source/MediaPlayerEditor/Private/Toolkits/MediaPlayerEditorToolkitMediaPlayerBase.cpp
@@ -75,7 +75,7 @@ float FMediaPlayerEditorToolkitMediaPlayerBase::GetForwardRate() const
 		Rate = 1.0f;
 	}
 
-	return 2.0f * Rate;
+	return 2 * Rate;
 }
 
 float FMediaPlayerEditorToolkitMediaPlayerBase::GetReverseRate() const
@@ -87,7 +87,7 @@ float FMediaPlayerEditorToolkitMediaPlayerBase::GetReverseRate() const
 		return -1.0f;
 	}
 
-	return 2.0f * Rate;
+	return 2 * Rate;
 }
 
 void FMediaPlayerEditorToolkitMediaPlayerBase::BindCommands()
@@ -102,8 +102,14 @@ void FMediaPlayerEditorToolkitMediaPlayerBase::BindCommands()
 
 	ToolkitCommands->MapAction(
 		Commands.ForwardMedia,
-		FExecuteAction::CreateLambda([this]{ MediaPlayer->SetRate(GetForwardRate()); }),
-		FCanExecuteAction::CreateLambda([this]{ return MediaPlayer->IsReady() && MediaPlayer->SupportsRate(GetForwardRate(), false); })
+		FExecuteAction::CreateLambda([this] { MediaPlayer->SetRate(GetForwardRate()); }),
+		FCanExecuteAction::CreateLambda([this] { return MediaPlayer->IsReady() && MediaPlayer->SupportsRate(GetForwardRate(), false); })
+	);
+
+	ToolkitCommands->MapAction(
+		Commands.PlayReverseMedia,
+		FExecuteAction::CreateLambda([this] { MediaPlayer->SetRate(-1.f); }),
+		FCanExecuteAction::CreateLambda([this] { return MediaPlayer->IsReady() && MediaPlayer->SupportsRate(-1.f,  /* Unthinned */ true) && (!MediaPlayer->IsPlaying() || (MediaPlayer->GetRate() != -1.0f)); })
 	);
 
 	ToolkitCommands->MapAction(
@@ -138,7 +144,24 @@ void FMediaPlayerEditorToolkitMediaPlayerBase::BindCommands()
 					MediaPlayer->Play();
 				}
 			}),
-		FCanExecuteAction::CreateLambda([this] { return MediaPlayer->IsReady(); })
+		FCanExecuteAction::CreateLambda([this] { return MediaPlayer->IsReady() && (MediaPlayer->IsPlaying() || MediaPlayer->GetRate() != 1.f); })
+	);
+
+	ToolkitCommands->MapAction(
+		Commands.TogglePlayReversePauseMedia,
+		FExecuteAction::CreateLambda(
+			[this]
+			{
+				if (MediaPlayer->IsPlaying())
+				{
+					MediaPlayer->Pause();
+				}
+				else
+				{
+					MediaPlayer->SetRate(-1.f);
+				}
+			}),
+		FCanExecuteAction::CreateLambda([this] { return MediaPlayer->IsReady() && MediaPlayer->SupportsRate(-1.f, /* Unthinned */ true) && (MediaPlayer->IsPlaying() || MediaPlayer->GetRate() != -1.f ); })
 	);
 
 	ToolkitCommands->MapAction(
@@ -149,8 +172,8 @@ void FMediaPlayerEditorToolkitMediaPlayerBase::BindCommands()
 
 	ToolkitCommands->MapAction(
 		Commands.ReverseMedia,
-		FExecuteAction::CreateLambda([this]{ MediaPlayer->SetRate(GetReverseRate()); } ),
-		FCanExecuteAction::CreateLambda([this]{ return MediaPlayer->IsReady() && MediaPlayer->SupportsRate(GetReverseRate(), false); })
+		FExecuteAction::CreateLambda([this] { MediaPlayer->SetRate(GetReverseRate()); }),
+		FCanExecuteAction::CreateLambda([this] { return MediaPlayer->IsReady() && MediaPlayer->SupportsRate(GetReverseRate(), false); })
 	);
 
 	ToolkitCommands->MapAction(
@@ -160,7 +183,7 @@ void FMediaPlayerEditorToolkitMediaPlayerBase::BindCommands()
 	);
 
 	ToolkitCommands->MapAction(
-		Commands.FastForwardMedia,
+		Commands.JumpToEndMedia,
 		FExecuteAction::CreateLambda([this] { MediaPlayer->Seek(MediaPlayer->GetDuration()); }),
 		FCanExecuteAction::CreateLambda([this] { return MediaPlayer->IsReady() && MediaPlayer->SupportsSeeking(); })
 	);
@@ -168,7 +191,7 @@ void FMediaPlayerEditorToolkitMediaPlayerBase::BindCommands()
 	static const FTimespan TenSeconds = FTimespan::FromSeconds(10.f);
 
 	ToolkitCommands->MapAction(
-		Commands.StepBackward,
+		Commands.StepBackwardMedia,
 		FExecuteAction::CreateLambda([this] 
 			{
 				// TODO: Convert to underlying media player command when framework support is added.
@@ -184,26 +207,7 @@ void FMediaPlayerEditorToolkitMediaPlayerBase::BindCommands()
 	);
 
 	ToolkitCommands->MapAction(
-		Commands.JumpBackward,
-		FExecuteAction::CreateLambda(
-			[this]
-			{
-				const FTimespan CurrentTime = MediaPlayer->GetTime();
-
-				if (CurrentTime <= TenSeconds)
-				{
-					MediaPlayer->Seek(0);
-				}
-				else
-				{
-					MediaPlayer->Seek(MediaPlayer->GetTime() - FTimespan::FromSeconds(10.f));
-				}
-			}),
-		FCanExecuteAction::CreateLambda([this] { return MediaPlayer->IsReady() && MediaPlayer->SupportsSeeking(); })
-	);
-
-	ToolkitCommands->MapAction(
-		Commands.StepForward,
+		Commands.StepForwardMedia,
 		FExecuteAction::CreateLambda([this] 
 			{
 				// TODO: Convert to underlying media player command when framework support is added.
@@ -218,26 +222,6 @@ void FMediaPlayerEditorToolkitMediaPlayerBase::BindCommands()
 		FCanExecuteAction::CreateLambda([this] { return MediaPlayer->IsReady() && MediaPlayer->IsPaused() && MediaPlayer->SupportsSeeking(); })
 	);
 
-	ToolkitCommands->MapAction(
-		Commands.JumpForward,
-		FExecuteAction::CreateLambda(
-			[this]
-			{
-				const FTimespan CurrentTime = MediaPlayer->GetTime();
-				const FTimespan Duration = MediaPlayer->GetDuration();
-
-				if ((Duration - CurrentTime) <= TenSeconds)
-				{
-					MediaPlayer->Seek(Duration);
-				}
-				else
-				{
-					MediaPlayer->Seek(MediaPlayer->GetTime() + FTimespan::FromSeconds(10.f));
-				}
-			}),
-		FCanExecuteAction::CreateLambda([this] { return MediaPlayer->IsReady() && MediaPlayer->SupportsSeeking(); })
-	);
-
 	ToolkitCommands->MapAction(
 		Commands.ToggledRedTextureChannel,
 		FExecuteAction::CreateLambda([this] { ToggleChannelMask(MediaPlayerEditor::MediaImage::ETextureChannelMask::Red); }),
diff --git a/Engine/Plugins/Media/MediaPlayerEditor/Source/MediaPlayerEditor/Private/Toolkits/MediaSourceEditorToolkit.cpp b/Engine/Plugins/Media/MediaPlayerEditor/Source/MediaPlayerEditor/Private/Toolkits/MediaSourceEditorToolkit.cpp
--- a/Engine/Plugins/Media/MediaPlayerEditor/Source/MediaPlayerEditor/Private/Toolkits/MediaSourceEditorToolkit.cpp
+++ b/Engine/Plugins/Media/MediaPlayerEditor/Source/MediaPlayerEditor/Private/Toolkits/MediaSourceEditorToolkit.cpp
@@ -246,22 +246,8 @@ void FMediaSourceEditorToolkit::ExtendToolBar()
 			{
 				const FMediaPlayerEditorCommands& Commands = FMediaPlayerEditorCommands::Get();
 
-				ToolbarBuilder.BeginSection("PlaybackControls");
-				{
-					ToolbarBuilder.AddToolBarButton(Commands.PreviousMedia);
-					ToolbarBuilder.AddToolBarButton(Commands.RewindMedia);
-					ToolbarBuilder.AddToolBarButton(Commands.ReverseMedia);
-					ToolbarBuilder.AddToolBarButton(Commands.PlayMedia);
-					ToolbarBuilder.AddToolBarButton(Commands.PauseMedia);
-					ToolbarBuilder.AddToolBarButton(Commands.ForwardMedia);
-					ToolbarBuilder.AddToolBarButton(Commands.NextMedia);
-				}
-				ToolbarBuilder.EndSection();
-
 				ToolbarBuilder.BeginSection("MediaControls");
 				{
-					ToolbarBuilder.AddToolBarButton(Commands.OpenMedia);
-					ToolbarBuilder.AddToolBarButton(Commands.CloseMedia);
 					ToolbarBuilder.AddToolBarButton(Commands.GenerateThumbnail);
 				}
 				ToolbarBuilder.EndSection();
diff --git a/Engine/Plugins/Media/MediaPlayerEditor/Source/MediaPlayerEditor/Private/Widgets/SMediaControls.cpp b/Engine/Plugins/Media/MediaPlayerEditor/Source/MediaPlayerEditor/Private/Widgets/SMediaControls.cpp
new file mode 100644
--- /dev/null
+++ b/Engine/Plugins/Media/MediaPlayerEditor/Source/MediaPlayerEditor/Private/Widgets/SMediaControls.cpp
@@ -0,0 +1,653 @@
+// Copyright Epic Games, Inc. All Rights Reserved.
+
+#include "Widgets/SMediaControls.h"
+
+#include "Framework/Commands/UICommandInfo.h"
+#include "Framework/Commands/UICommandList.h"
+#include "Internationalization/Text.h"
+#include "MediaPlayer.h"
+#include "MediaPlayerEditorModule.h"
+#include "Models/MediaPlayerEditorCommands.h"
+#include "Modules/ModuleManager.h"
+#include "Shared/MediaPlayerEditorStyle.h"
+#include "Widgets/Images/SImage.h"
+#include "Widgets/Input/SButton.h"
+#include "Widgets/SBoxPanel.h"
+#include "SMediaControls.h"
+
+#define LOCTEXT_NAMESPACE "SMediaControls"
+
+namespace MediaPlayerEditor
+{
+	namespace Private
+	{
+		FText GetKeybindText(const TSharedPtr<FUICommandInfo>& InCommand)
+		{
+			if (InCommand.IsValid())
+			{
+				const TSharedRef<const FInputChord>& Chord = InCommand->GetFirstValidChord();
+
+				if (Chord->IsValidChord())
+				{
+					return Chord->GetInputText(/* Long display name */ false);
+				}
+			}
+
+			return FText::GetEmpty();
+		}
+
+		FText GetToolTipWithKeybind(const TSharedRef<FUICommandInfo>& InCommand, TSharedPtr<FUICommandInfo> InBackupKeybindCommand = nullptr)
+		{
+			FText KeybindText = GetKeybindText(InCommand);
+
+			if (KeybindText.IsEmpty())
+			{
+				KeybindText = GetKeybindText(InBackupKeybindCommand);
+			}
+
+			const FText& Description = InCommand->GetDescription();
+			const FText& Text = Description.IsEmpty() ? InCommand->GetLabel() : Description;
+
+			if (KeybindText.IsEmpty())
+			{
+				return Text;
+			}
+
+			return FText::Format(LOCTEXT("GetToolTipWithKeybind", "{0}\n\nKey: {1}"), Text, KeybindText);
+		}
+	}
+}
+
+void SMediaControls::Construct(const FArguments& InArgs, TSharedPtr<FUICommandList> InCommandList, UMediaPlayer* InMediaPlayer)
+{
+	CommandListWeak = InCommandList;
+	PlayerWeak = InMediaPlayer;
+
+	ChildSlot
+	[
+		CreateControls()
+	];
+}
+
+TSharedRef<SWidget> SMediaControls::CreateControls()
+{
+	using namespace MediaPlayerEditor::Private;
+
+	const FMediaPlayerEditorCommands& Commands = FMediaPlayerEditorCommands::Get();
+	
+	IMediaPlayerEditorModule& Module = FModuleManager::GetModuleChecked<IMediaPlayerEditorModule>("MediaPlayerEditor");	
+	TSharedPtr<ISlateStyle> Style = Module.GetStyle();
+
+	return SNew(SHorizontalBox)
+
+		// Rewind button.
+		+ SHorizontalBox::Slot()
+		.VAlign(VAlign_Center)
+		.AutoWidth()
+		[
+			SNew(SButton)
+			.VAlign(VAlign_Center)
+			.IsEnabled(this, &SMediaControls::Rewind_IsEnabled)
+			.OnClicked(this, &SMediaControls::Rewind_OnClicked)
+			.ButtonStyle(Style, "MediaPlayerEditor.MediaControlButton")
+			.ContentPadding(2.f)
+			[
+				SNew(SImage)
+				.ColorAndOpacity(FSlateColor::UseForeground())
+				.Image(Style->GetBrush("MediaPlayerEditor.RewindMedia"))
+				.ToolTipText(GetToolTipWithKeybind(Commands.RewindMedia.ToSharedRef()))
+				.DesiredSizeOverride(FVector2D(20.f))
+			]
+		]
+		
+		// Reverse button.
+		/*
+		+ SHorizontalBox::Slot()
+		.VAlign(VAlign_Center)
+		.AutoWidth()
+		[
+			SNew(SButton)
+			.VAlign(VAlign_Center)
+			.IsEnabled(this, &SMediaControls::Reverse_IsEnabled)
+			.OnClicked(this, &SMediaControls::Reverse_OnClicked)
+			.ButtonStyle(Style, "MediaPlayerEditor.MediaControlButton")
+			.ContentPadding(2.f)
+			[
+				SNew(SImage)
+				.ColorAndOpacity(FSlateColor::UseForeground())
+				.Image(Style->GetBrush("MediaPlayerEditor.ReverseMedia"))
+				.ToolTipText(GetToolTipWithKeybind(Commands.ReverseMedia.ToSharedRef()))
+				.DesiredSizeOverride(FVector2D(20.f))
+			]
+		]
+		*/
+		
+		// Step back button.
+		+ SHorizontalBox::Slot()
+		.VAlign(VAlign_Center)
+		.AutoWidth()
+		[
+			SNew(SButton)
+			.VAlign(VAlign_Center)
+			.IsEnabled(this, &SMediaControls::StepBack_IsEnabled)
+			.OnClicked(this, &SMediaControls::StepBack_OnClicked)
+			.ButtonStyle(Style, "MediaPlayerEditor.MediaControlButton")
+			.ContentPadding(2.f)
+			[
+				SNew(SImage)
+				.ColorAndOpacity(FSlateColor::UseForeground())
+				.Image(Style->GetBrush("MediaPlayerEditor.StepBackwardMedia"))
+				.ToolTipText(GetToolTipWithKeybind(Commands.StepBackwardMedia.ToSharedRef()))
+				.DesiredSizeOverride(FVector2D(20.f))
+			]
+		]
+
+		// Reverse button.
+		+ SHorizontalBox::Slot()
+		.VAlign(VAlign_Center)
+		.AutoWidth()
+		[
+			SNew(SButton)
+			.VAlign(VAlign_Center)
+			.IsEnabled(this, &SMediaControls::PlayReverse_IsEnabled)
+			.OnClicked(this, &SMediaControls::PlayReverse_OnClicked)
+			.ButtonStyle(Style, "MediaPlayerEditor.MediaControlButton")
+			.ContentPadding(2.f)
+			[
+				SNew(SImage)
+				.ColorAndOpacity(FSlateColor::UseForeground())
+				.Image(this, &SMediaControls::PlayReverse_GetBrush)
+				.ToolTipText(this, &SMediaControls::PlayReverse_GetToolTip)
+				.DesiredSizeOverride(FVector2D(20.f))
+			]
+		]
+
+		// Play button.
+		+ SHorizontalBox::Slot()
+		.VAlign(VAlign_Center)
+		.AutoWidth()
+		[
+			SNew(SButton)
+			.VAlign(VAlign_Center)
+			.IsEnabled(this, &SMediaControls::Play_IsEnabled)
+			.OnClicked(this, &SMediaControls::Play_OnClicked)
+			.ButtonStyle(Style, "MediaPlayerEditor.MediaControlButton")
+			.ContentPadding(2.f)
+			[
+				SNew(SImage)
+				.ColorAndOpacity(FSlateColor::UseForeground())
+				.Image(this, &SMediaControls::Play_GetBrush)
+				.ToolTipText(this, &SMediaControls::Play_GetToolTip)
+				.DesiredSizeOverride(FVector2D(20.f))
+			]
+		]
+
+		// Step forward button.
+		+ SHorizontalBox::Slot()
+		.VAlign(VAlign_Center)
+		.AutoWidth()
+		[
+			SNew(SButton)
+			.VAlign(VAlign_Center)
+			.IsEnabled(this, &SMediaControls::StepForwardMedia_IsEnabled)
+			.OnClicked(this, &SMediaControls::StepForwardMedia_OnClicked)
+			.ButtonStyle(Style, "MediaPlayerEditor.MediaControlButton")
+			.ContentPadding(2.f)
+			[
+				SNew(SImage)
+				.ColorAndOpacity(FSlateColor::UseForeground())
+				.Image(Style->GetBrush("MediaPlayerEditor.StepForwardMedia"))
+				.ToolTipText(GetToolTipWithKeybind(Commands.StepForwardMedia.ToSharedRef()))
+				.DesiredSizeOverride(FVector2D(20.f))
+			]
+		]
+		
+		// Fast forward button.
+		/*
+		+ SHorizontalBox::Slot()
+		.VAlign(VAlign_Center)
+		.AutoWidth()
+		[
+			SNew(SButton)
+			.VAlign(VAlign_Center)
+			.IsEnabled(this, &SMediaControls::Forward_IsEnabled)
+			.OnClicked(this, &SMediaControls::Forward_OnClicked)
+			.ButtonStyle(Style, "MediaPlayerEditor.MediaControlButton")
+			.ContentPadding(2.f)
+			[
+				SNew(SImage)
+				.ColorAndOpacity(FSlateColor::UseForeground())
+				.Image(Style->GetBrush("MediaPlayerEditor.ForwardMedia"))
+				.ToolTipText(GetToolTipWithKeybind(Commands.ForwardMedia.ToSharedRef()))
+				.DesiredSizeOverride(FVector2D(20.f))
+			]
+		]
+		*/
+
+		// Jump to end button.
+		+ SHorizontalBox::Slot()
+		.VAlign(VAlign_Center)
+		.AutoWidth()
+		[
+			SNew(SButton)
+			.VAlign(VAlign_Center)
+			.IsEnabled(this, &SMediaControls::JumpToEndMedia_IsEnabled)
+			.OnClicked(this, &SMediaControls::JumpToEndMedia_OnClicked)
+			.ButtonStyle(Style, "MediaPlayerEditor.MediaControlButton")
+			.ContentPadding(2.f)
+			[
+				SNew(SImage)
+				.ColorAndOpacity(FSlateColor::UseForeground())
+				.Image(Style->GetBrush("MediaPlayerEditor.JumpToEndMedia"))
+				.ToolTipText(GetToolTipWithKeybind(Commands.JumpToEndMedia.ToSharedRef()))
+				.DesiredSizeOverride(FVector2D(20.f))
+			]
+		]
+
+		// Previous button
+		+ SHorizontalBox::Slot()
+		.VAlign(VAlign_Center)
+		.AutoWidth()
+		.Padding(20.f, 0.f, 0.f, 0.f)
+		[
+			SNew(SButton)
+			.VAlign(VAlign_Center)
+			.IsEnabled(this, &SMediaControls::PreviousMedia_IsEnabled)
+			.OnClicked(this, &SMediaControls::PreviousMedia_OnClicked)
+			.ButtonStyle(Style, "MediaPlayerEditor.MediaControlButton")
+			.ContentPadding(2.f)
+			[
+				SNew(SImage)
+				.ColorAndOpacity(FSlateColor::UseForeground())
+				.Image(Style->GetBrush("MediaPlayerEditor.PreviousMedia"))
+				.ToolTipText(GetToolTipWithKeybind(Commands.PreviousMedia.ToSharedRef()))
+				.DesiredSizeOverride(FVector2D(20.f))
+			]
+		]
+
+		// Next button
+		+ SHorizontalBox::Slot()
+		.VAlign(VAlign_Center)
+		.AutoWidth()
+		[
+			SNew(SButton)
+			.VAlign(VAlign_Center)
+			.IsEnabled(this, &SMediaControls::NextMedia_IsEnabled)
+			.OnClicked(this, &SMediaControls::NextMedia_OnClicked)
+			.ButtonStyle(Style, "MediaPlayerEditor.MediaControlButton")
+			.ContentPadding(2.f)
+			[
+				SNew(SImage)
+				.ColorAndOpacity(FSlateColor::UseForeground())
+				.Image(Style->GetBrush("MediaPlayerEditor.NextMedia"))
+				.ToolTipText(GetToolTipWithKeybind(Commands.NextMedia.ToSharedRef()))
+				.DesiredSizeOverride(FVector2D(20.f))
+			]
+		]
+
+		// Open/Close Media
+		+ SHorizontalBox::Slot()
+		.VAlign(VAlign_Center)
+		.AutoWidth()
+		[
+			SNew(SButton)
+			.VAlign(VAlign_Center)
+			.IsEnabled(this, &SMediaControls::OpenCloseMedia_IsEnabled)
+			.OnClicked(this, &SMediaControls::OpenCloseMedia_OnClicked)
+			.ButtonStyle(Style, "MediaPlayerEditor.MediaControlButton")
+			.ContentPadding(2.f)
+			[
+				SNew(SImage)
+				.ColorAndOpacity(FSlateColor::UseForeground())
+				.Image(this, &SMediaControls::OpenCloseMedia_GetBrush)
+				.ToolTipText(this, &SMediaControls::OpenCloseMedia_GetToolTip)
+				.DesiredSizeOverride(FVector2D(20.f))
+			]
+		];
+}
+
+bool SMediaControls::PreviousMedia_IsEnabled() const
+{
+	if (TSharedPtr<FUICommandList> CommandList = CommandListWeak.Pin())
+	{
+		return CommandList->CanExecuteAction(FMediaPlayerEditorCommands::Get().PreviousMedia.ToSharedRef());
+	}
+
+	return false;
+}
+
+FReply SMediaControls::PreviousMedia_OnClicked()
+{
+	if (TSharedPtr<FUICommandList> CommandList = CommandListWeak.Pin())
+	{
+		CommandList->TryExecuteAction(FMediaPlayerEditorCommands::Get().PreviousMedia.ToSharedRef());
+	}
+
+	return FReply::Handled();
+}
+
+bool SMediaControls::Rewind_IsEnabled() const
+{
+	if (TSharedPtr<FUICommandList> CommandList = CommandListWeak.Pin())
+	{
+		return CommandList->CanExecuteAction(FMediaPlayerEditorCommands::Get().RewindMedia.ToSharedRef());
+	}
+
+	return false;
+}
+
+FReply SMediaControls::Rewind_OnClicked()
+{
+	if (TSharedPtr<FUICommandList> CommandList = CommandListWeak.Pin())
+	{
+		CommandList->TryExecuteAction(FMediaPlayerEditorCommands::Get().RewindMedia.ToSharedRef());
+	}
+
+	return FReply::Handled();
+}
+
+bool SMediaControls::PlayReverse_IsEnabled() const
+{
+	if (TSharedPtr<FUICommandList> CommandList = CommandListWeak.Pin())
+	{
+		return CommandList->CanExecuteAction(FMediaPlayerEditorCommands::Get().TogglePlayReversePauseMedia.ToSharedRef());
+	}
+
+	return false;
+}
+
+const FSlateBrush* SMediaControls::PlayReverse_GetBrush() const
+{
+	IMediaPlayerEditorModule& Module = FModuleManager::GetModuleChecked<IMediaPlayerEditorModule>("MediaPlayerEditor");
+	TSharedPtr<ISlateStyle> Style = Module.GetStyle();
+
+	if (UMediaPlayer* MediaPlayer = PlayerWeak.Get())
+	{
+		const float Rate = MediaPlayer->GetRate();
+
+		if (Rate < 0 && !FMath::IsNearlyZero(Rate))
+		{
+			return Style->GetBrush("MediaPlayerEditor.PauseMedia");
+		}
+	}
+
+	return Style->GetBrush("MediaPlayerEditor.PlayReverseMedia");
+}
+
+FText SMediaControls::PlayReverse_GetToolTip() const
+{
+	using namespace MediaPlayerEditor::Private;
+
+	const FMediaPlayerEditorCommands& Commands = FMediaPlayerEditorCommands::Get();
+
+	if (UMediaPlayer* MediaPlayer = PlayerWeak.Get())
+	{
+		const float Rate = MediaPlayer->GetRate();
+
+		if (Rate < 0 && !FMath::IsNearlyZero(Rate))
+		{
+			return GetToolTipWithKeybind(Commands.PauseMedia.ToSharedRef(), Commands.TogglePlayReversePauseMedia);
+		}
+	}
+
+	return GetToolTipWithKeybind(Commands.PlayReverseMedia.ToSharedRef(), Commands.TogglePlayReversePauseMedia);
+}
+
+FReply SMediaControls::PlayReverse_OnClicked()
+{
+	if (TSharedPtr<FUICommandList> CommandList = CommandListWeak.Pin())
+	{
+		CommandList->TryExecuteAction(FMediaPlayerEditorCommands::Get().TogglePlayReversePauseMedia.ToSharedRef());
+	}
+
+	return FReply::Handled();
+}
+
+bool SMediaControls::Reverse_IsEnabled() const
+{
+	if (TSharedPtr<FUICommandList> CommandList = CommandListWeak.Pin())
+	{
+		return CommandList->CanExecuteAction(FMediaPlayerEditorCommands::Get().ReverseMedia.ToSharedRef());
+	}
+
+	return false;
+}
+
+FReply SMediaControls::Reverse_OnClicked()
+{
+	if (TSharedPtr<FUICommandList> CommandList = CommandListWeak.Pin())
+	{
+		CommandList->TryExecuteAction(FMediaPlayerEditorCommands::Get().ReverseMedia.ToSharedRef());
+	}
+
+	return FReply::Handled();
+}
+
+bool SMediaControls::StepBack_IsEnabled() const
+{
+	if (TSharedPtr<FUICommandList> CommandList = CommandListWeak.Pin())
+	{
+		return CommandList->CanExecuteAction(FMediaPlayerEditorCommands::Get().StepBackwardMedia.ToSharedRef());
+	}
+
+	return false;
+}
+
+FReply SMediaControls::StepBack_OnClicked()
+{
+	if (TSharedPtr<FUICommandList> CommandList = CommandListWeak.Pin())
+	{
+		CommandList->TryExecuteAction(FMediaPlayerEditorCommands::Get().StepBackwardMedia.ToSharedRef());
+	}
+
+	return FReply::Handled();
+}
+
+bool SMediaControls::Play_IsEnabled() const
+{
+	if (TSharedPtr<FUICommandList> CommandList = CommandListWeak.Pin())
+	{
+		return CommandList->CanExecuteAction(FMediaPlayerEditorCommands::Get().TogglePlayPauseMedia.ToSharedRef());
+	}
+
+	return false;
+}
+
+const FSlateBrush* SMediaControls::Play_GetBrush() const
+{
+	IMediaPlayerEditorModule& Module = FModuleManager::GetModuleChecked<IMediaPlayerEditorModule>("MediaPlayerEditor");
+	TSharedPtr<ISlateStyle> Style = Module.GetStyle();
+
+	if (UMediaPlayer* MediaPlayer = PlayerWeak.Get())
+	{
+		const float Rate = MediaPlayer->GetRate();
+
+		if (Rate > 0 && !FMath::IsNearlyZero(Rate))
+		{
+			return Style->GetBrush("MediaPlayerEditor.PauseMedia");
+		}
+	}
+
+	return Style->GetBrush("MediaPlayerEditor.PlayMedia");
+}
+
+FText SMediaControls::Play_GetToolTip() const
+{
+	using namespace MediaPlayerEditor::Private;
+
+	const FMediaPlayerEditorCommands& Commands = FMediaPlayerEditorCommands::Get();
+
+	if (UMediaPlayer* MediaPlayer = PlayerWeak.Get())
+	{
+		const float Rate = MediaPlayer->GetRate();
+
+		if (Rate > 0 && !FMath::IsNearlyZero(Rate))
+		{
+			return GetToolTipWithKeybind(Commands.PauseMedia.ToSharedRef(), Commands.TogglePlayPauseMedia);
+		}
+	}
+
+	return GetToolTipWithKeybind(Commands.PlayMedia.ToSharedRef(), Commands.TogglePlayPauseMedia);
+}
+
+FReply SMediaControls::Play_OnClicked()
+{
+	if (TSharedPtr<FUICommandList> CommandList = CommandListWeak.Pin())
+	{
+		CommandList->TryExecuteAction(FMediaPlayerEditorCommands::Get().TogglePlayPauseMedia.ToSharedRef());
+	}
+
+	return FReply::Handled();
+}
+
+bool SMediaControls::StepForwardMedia_IsEnabled() const
+{
+	if (TSharedPtr<FUICommandList> CommandList = CommandListWeak.Pin())
+	{
+		return CommandList->CanExecuteAction(FMediaPlayerEditorCommands::Get().StepForwardMedia.ToSharedRef());
+	}
+
+	return false;
+}
+
+FReply SMediaControls::StepForwardMedia_OnClicked()
+{
+	if (TSharedPtr<FUICommandList> CommandList = CommandListWeak.Pin())
+	{
+		CommandList->TryExecuteAction(FMediaPlayerEditorCommands::Get().StepForwardMedia.ToSharedRef());
+	}
+
+	return FReply::Handled();
+}
+
+bool SMediaControls::Forward_IsEnabled() const
+{
+	if (TSharedPtr<FUICommandList> CommandList = CommandListWeak.Pin())
+	{
+		return CommandList->CanExecuteAction(FMediaPlayerEditorCommands::Get().ForwardMedia.ToSharedRef());
+	}
+
+	return false;
+}
+
+FReply SMediaControls::Forward_OnClicked()
+{
+	if (TSharedPtr<FUICommandList> CommandList = CommandListWeak.Pin())
+	{
+		CommandList->TryExecuteAction(FMediaPlayerEditorCommands::Get().ForwardMedia.ToSharedRef());
+	}
+
+	return FReply::Handled();
+}
+
+bool SMediaControls::JumpToEndMedia_IsEnabled() const
+{
+	if (TSharedPtr<FUICommandList> CommandList = CommandListWeak.Pin())
+	{
+		return CommandList->CanExecuteAction(FMediaPlayerEditorCommands::Get().JumpToEndMedia.ToSharedRef());
+	}
+
+	return false;
+}
+
+FReply SMediaControls::JumpToEndMedia_OnClicked()
+{
+	if (TSharedPtr<FUICommandList> CommandList = CommandListWeak.Pin())
+	{
+		CommandList->TryExecuteAction(FMediaPlayerEditorCommands::Get().JumpToEndMedia.ToSharedRef());
+	}
+
+	return FReply::Handled();
+}
+
+bool SMediaControls::NextMedia_IsEnabled() const
+{
+	if (TSharedPtr<FUICommandList> CommandList = CommandListWeak.Pin())
+	{
+		return CommandList->CanExecuteAction(FMediaPlayerEditorCommands::Get().NextMedia.ToSharedRef());
+	}
+
+	return false;
+}
+
+FReply SMediaControls::NextMedia_OnClicked()
+{
+	if (TSharedPtr<FUICommandList> CommandList = CommandListWeak.Pin())
+	{
+		CommandList->TryExecuteAction(FMediaPlayerEditorCommands::Get().NextMedia.ToSharedRef());
+	}
+
+	return FReply::Handled();
+}
+
+bool SMediaControls::OpenCloseMedia_IsEnabled() const
+{
+	if (TSharedPtr<FUICommandList> CommandList = CommandListWeak.Pin())
+	{
+		const FMediaPlayerEditorCommands& Commands = FMediaPlayerEditorCommands::Get();
+
+		if (UMediaPlayer* MediaPlayer = PlayerWeak.Get())
+		{
+			if (MediaPlayer->IsReady())
+			{
+				return CommandList->CanExecuteAction(Commands.CloseMedia.ToSharedRef());
+			}
+		}
+
+		return CommandList->CanExecuteAction(Commands.OpenMedia.ToSharedRef());
+	}
+
+	return false;
+}
+
+const FSlateBrush* SMediaControls::OpenCloseMedia_GetBrush() const
+{
+	IMediaPlayerEditorModule& Module = FModuleManager::GetModuleChecked<IMediaPlayerEditorModule>("MediaPlayerEditor");
+	TSharedPtr<ISlateStyle> Style = Module.GetStyle();
+
+	if (UMediaPlayer* MediaPlayer = PlayerWeak.Get())
+	{
+		if (MediaPlayer->IsReady())
+		{
+			return Style->GetBrush("MediaPlayerEditor.CloseMedia");
+		}
+	}
+
+	return Style->GetBrush("MediaPlayerEditor.OpenMedia");
+}
+
+FText SMediaControls::OpenCloseMedia_GetToolTip() const
+{
+	using namespace MediaPlayerEditor::Private;
+
+	if (UMediaPlayer* MediaPlayer = PlayerWeak.Get())
+	{
+		if (MediaPlayer->IsReady())
+		{
+			return GetToolTipWithKeybind(FMediaPlayerEditorCommands::Get().CloseMedia.ToSharedRef());
+		}
+	}
+
+	return GetToolTipWithKeybind(FMediaPlayerEditorCommands::Get().OpenMedia.ToSharedRef());
+}
+
+FReply SMediaControls::OpenCloseMedia_OnClicked()
+{
+	if (TSharedPtr<FUICommandList> CommandList = CommandListWeak.Pin())
+	{
+		if (UMediaPlayer* MediaPlayer = PlayerWeak.Get())
+		{
+			if (MediaPlayer->IsReady())
+			{
+				CommandList->TryExecuteAction(FMediaPlayerEditorCommands::Get().CloseMedia.ToSharedRef());
+				return FReply::Handled();
+			}
+		}
+
+		CommandList->TryExecuteAction(FMediaPlayerEditorCommands::Get().OpenMedia.ToSharedRef());
+	}
+
+	return FReply::Handled();
+}
+
+#undef LOCTEXT_NAMESPACE
diff --git a/Engine/Plugins/Media/MediaPlayerEditor/Source/MediaPlayerEditor/Private/Widgets/SMediaControls.h b/Engine/Plugins/Media/MediaPlayerEditor/Source/MediaPlayerEditor/Private/Widgets/SMediaControls.h
new file mode 100644
--- /dev/null
+++ b/Engine/Plugins/Media/MediaPlayerEditor/Source/MediaPlayerEditor/Private/Widgets/SMediaControls.h
@@ -0,0 +1,66 @@
+// Copyright Epic Games, Inc. All Rights Reserved.
+
+#pragma once
+
+#include "Widgets/SCompoundWidget.h"
+
+#include "UObject/WeakObjectPtr.h"
+
+class FText;
+class FUICommandList;
+class UMediaPlayer;
+
+class SMediaControls : public SCompoundWidget
+{
+public:
+	SLATE_BEGIN_ARGS(SMediaControls) 
+		{}
+	SLATE_END_ARGS()
+
+	void Construct(const FArguments& InArgs, TSharedPtr<FUICommandList> InCommandList, UMediaPlayer* InMediaPlayer);
+
+protected:
+	TWeakPtr<FUICommandList> CommandListWeak;
+	TWeakObjectPtr<UMediaPlayer> PlayerWeak;
+
+	TSharedRef<SWidget> CreateControls();
+
+	bool PreviousMedia_IsEnabled() const;
+	FReply PreviousMedia_OnClicked();
+
+	bool Rewind_IsEnabled() const;
+	FReply Rewind_OnClicked();
+
+	bool PlayReverse_IsEnabled() const;
+	const FSlateBrush* PlayReverse_GetBrush() const;
+	FText PlayReverse_GetToolTip() const;
+	FReply PlayReverse_OnClicked();
+
+	bool Reverse_IsEnabled() const;
+	FReply Reverse_OnClicked();
+
+	bool StepBack_IsEnabled() const;
+	FReply StepBack_OnClicked();
+
+	bool Play_IsEnabled() const;
+	const FSlateBrush* Play_GetBrush() const;
+	FText Play_GetToolTip() const;
+	FReply Play_OnClicked();
+
+	bool StepForwardMedia_IsEnabled() const;
+	FReply StepForwardMedia_OnClicked();
+
+	bool Forward_IsEnabled() const;
+	FReply Forward_OnClicked();
+
+	bool JumpToEndMedia_IsEnabled() const;
+	FReply JumpToEndMedia_OnClicked();
+
+	bool NextMedia_IsEnabled() const;
+	FReply NextMedia_OnClicked();
+
+	bool OpenCloseMedia_IsEnabled() const;
+	const FSlateBrush* OpenCloseMedia_GetBrush() const;
+	FText OpenCloseMedia_GetToolTip() const;
+	FReply OpenCloseMedia_OnClicked();
+};
\ No newline at end of file
diff --git a/Engine/Plugins/Media/MediaPlayerEditor/Source/MediaPlayerEditor/Private/Widgets/SMediaPlayerEditorViewer.cpp b/Engine/Plugins/Media/MediaPlayerEditor/Source/MediaPlayerEditor/Private/Widgets/SMediaPlayerEditorViewer.cpp
--- a/Engine/Plugins/Media/MediaPlayerEditor/Source/MediaPlayerEditor/Private/Widgets/SMediaPlayerEditorViewer.cpp
+++ b/Engine/Plugins/Media/MediaPlayerEditor/Source/MediaPlayerEditor/Private/Widgets/SMediaPlayerEditorViewer.cpp
@@ -37,6 +37,11 @@
 #include "SMediaPlayerEditorCache.h"
 #include "SMediaPlayerEditorViewport.h"
 
+#include "EditorWidgetsModule.h"
+#include "Models/MediaPlayerEditorCommands.h"
+#include "Modules/ModuleManager.h"
+#include "Widgets/SMediaControls.h"
+
 
 #define LOCTEXT_NAMESPACE "SMediaPlayerEditorViewer"
 
@@ -69,6 +74,7 @@ void SMediaPlayerEditorViewer::Construct(const FArguments& InArgs, UMediaPlayer&
 {
 	MediaPlayer = &InMediaPlayer;
 	Style = InStyle;
+	CommandListWeak = InArgs._Commands;
 
 	// initialize media player asset
 	MediaPlayer->OnMediaEvent().AddSP(this, &SMediaPlayerEditorViewer::HandleMediaPlayerMediaEvent);
@@ -408,48 +414,60 @@ void SMediaPlayerEditorViewer::Construct(const FArguments& InArgs, UMediaPlayer&
 									[
 										SNew(SHorizontalBox)
 
-										// timer
 										+ SHorizontalBox::Slot()
 											.AutoWidth()
-											.Padding(4.0f, 0.0f, 0.0f, 0.0f)
 											.VAlign(VAlign_Center)
+											.HAlign(HAlign_Left)
 											[
-												SNew(STextBlock)
-													.Text(this, &SMediaPlayerEditorViewer::HandleTimerTextBlockText)
-													.ToolTipText(this, &SMediaPlayerEditorViewer::HandleTimerTextBlockToolTipText)
-											]
+												SNew(SHorizontalBox)
+
+												// timer
+												+ SHorizontalBox::Slot()
+													.MinWidth(85.f)
+													.Padding(4.0f, 0.0f, 0.0f, 0.0f)
+													.VAlign(VAlign_Center)
+													.HAlign(HAlign_Right)
+													[
+														SNew(STextBlock)
+															.Text(this, &SMediaPlayerEditorViewer::HandleTimerTextBlockText)
+															.ToolTipText(this, &SMediaPlayerEditorViewer::HandleTimerTextBlockToolTipText)
+													]
 
-										// fps
-										+ SHorizontalBox::Slot()
-											.AutoWidth()
-											.Padding(8.0f, 0.0f, 4.0f, 0.0f)
-											.VAlign(VAlign_Center)
-											[
-												SNew(STextBlock)
-													.Text(this, &SMediaPlayerEditorViewer::HandleFpsTextBlockText)
-											]
+												// fps
+												+ SHorizontalBox::Slot()
+													.MinWidth(45.f)
+													.Padding(8.0f, 0.0f, 4.0f, 0.0f)
+													.VAlign(VAlign_Center)
+													.HAlign(HAlign_Left)
+													[
+														SNew(STextBlock)
+															.Text(this, &SMediaPlayerEditorViewer::HandleFpsTextBlockText)
+													]
 
-										// buffering indicator
-										+ SHorizontalBox::Slot()
-											.AutoWidth()
-											.Padding(4.0f, 0.0f)
-											.VAlign(VAlign_Center)
-											[
-												SNew(SThrobber)
-													.ToolTipText(LOCTEXT("BufferingTooltip", "Buffering..."))
-													.Visibility_Lambda([this]() -> EVisibility
-													{
-														return MediaPlayer->IsBuffering() ? EVisibility::Visible : EVisibility::Hidden;
-													})
+												// buffering indicator
+												+ SHorizontalBox::Slot()
+													.AutoWidth()
+													.Padding(4.0f, 0.0f)
+													.VAlign(VAlign_Center)
+													[
+														SNew(SThrobber)
+															.ToolTipText(LOCTEXT("BufferingTooltip", "Buffering..."))
+															.Visibility_Lambda([this]() -> EVisibility
+															{
+																return MediaPlayer->IsBuffering() ? EVisibility::Visible : EVisibility::Hidden;
+															})
+													]
 											]
 
+
 										// transport controls
 										+ SHorizontalBox::Slot()
 											.FillWidth(1.0f)
 											.Padding(8.0f, 0.0f)
 											.VAlign(VAlign_Center)
+											.HAlign(HAlign_Center)
 											[
-												SNew(SSpacer)//SNew(STransportControl)
+												CreatePlayerControls()
 											]
 
 										// settings
@@ -1142,6 +1160,9 @@ void SMediaPlayerEditorViewer::HandleVideoCaptureDevicesMenuNewMenu(FMenuBuilder
 	MakeCaptureDeviceMenu(DeviceInfos, MenuBuilder);
 }
 
-
+TSharedRef<SWidget> SMediaPlayerEditorViewer::CreatePlayerControls()
+{
+	return SNew(SMediaControls, CommandListWeak.Pin(), MediaPlayer);
+}
 
 #undef LOCTEXT_NAMESPACE
diff --git a/Engine/Plugins/Media/MediaPlayerEditor/Source/MediaPlayerEditor/Public/Models/MediaPlayerEditorCommands.h b/Engine/Plugins/Media/MediaPlayerEditor/Source/MediaPlayerEditor/Public/Models/MediaPlayerEditorCommands.h
--- a/Engine/Plugins/Media/MediaPlayerEditor/Source/MediaPlayerEditor/Public/Models/MediaPlayerEditorCommands.h
+++ b/Engine/Plugins/Media/MediaPlayerEditor/Source/MediaPlayerEditor/Public/Models/MediaPlayerEditorCommands.h
@@ -32,7 +32,7 @@ public:
 	/** Close the currently opened media. */
 	TSharedPtr<FUICommandInfo> CloseMedia;
 
-	/** Plays media forward. */
+	/** Plays media forward at ever increasing speed. */
 	TSharedPtr<FUICommandInfo> ForwardMedia;
 
 	/** Generate a thumnbnail. */
@@ -50,32 +50,32 @@ public:
 	/** Starts media playback. */
 	TSharedPtr<FUICommandInfo> PlayMedia;
 
+	/** Starts media playback in reverse. */
+	TSharedPtr<FUICommandInfo> PlayReverseMedia;
+
 	/** Toggles between media play and pause. */
 	TSharedPtr<FUICommandInfo> TogglePlayPauseMedia;
 
+	/** Toggles between media reverse play and pause. */
+	TSharedPtr<FUICommandInfo> TogglePlayReversePauseMedia;
+
 	/** Jump to previous item in the play list. */
 	TSharedPtr<FUICommandInfo> PreviousMedia;
 
-	/** Reverses media playback. */
+	/** Reverses media playback at ever increasing speed. */
 	TSharedPtr<FUICommandInfo> ReverseMedia;
 
 	/** Rewinds the media to the beginning. */
 	TSharedPtr<FUICommandInfo> RewindMedia;
 
 	/** Fast forwards the media to the end. */
-	TSharedPtr<FUICommandInfo> FastForwardMedia;
+	TSharedPtr<FUICommandInfo> JumpToEndMedia;
 
 	/** Step forward media 1 frame */
-	TSharedPtr<FUICommandInfo> StepForward;
+	TSharedPtr<FUICommandInfo> StepForwardMedia;
 
 	/** Step backward media 1 frame. */
-	TSharedPtr<FUICommandInfo> StepBackward;
-
-	/** Jump forward media 10s */
-	TSharedPtr<FUICommandInfo> JumpForward;
-
-	/** Jump backward media 10s. */
-	TSharedPtr<FUICommandInfo> JumpBackward;
+	TSharedPtr<FUICommandInfo> StepBackwardMedia;
 
 	/** Toggles red texture channel. */
 	TSharedPtr<FUICommandInfo> ToggledRedTextureChannel;
diff --git a/Engine/Plugins/Media/MediaPlayerEditor/Source/MediaPlayerEditor/Public/Widgets/SMediaPlayerEditorViewer.h b/Engine/Plugins/Media/MediaPlayerEditor/Source/MediaPlayerEditor/Public/Widgets/SMediaPlayerEditorViewer.h
--- a/Engine/Plugins/Media/MediaPlayerEditor/Source/MediaPlayerEditor/Public/Widgets/SMediaPlayerEditorViewer.h
+++ b/Engine/Plugins/Media/MediaPlayerEditor/Source/MediaPlayerEditor/Public/Widgets/SMediaPlayerEditorViewer.h
@@ -162,6 +162,9 @@ private:
 	/** Callback for creating the video capture devices sub-menu. */
 	UE_API void HandleVideoCaptureDevicesMenuNewMenu(FMenuBuilder& MenuBuilder);
 
+	/** Creates controls for the media */
+	UE_API TSharedRef<SWidget> CreatePlayerControls();
+
 private:
 
 	/** Whether something is currently being dragged over the widget. */
@@ -184,6 +187,9 @@ private:
 
 	/** The viewport that shows the media. */
 	TSharedPtr<SMediaPlayerEditorViewport> PlayerViewport;
+
+	/** The controls that can be executed by this viewer. */
+	TWeakPtr<FUICommandList> CommandListWeak;
 };
 
 #undef UE_API
