Media Assets: Added customization for stream media source assets that provides a dropdown for the stream URL property similar to the dropdown in the media player editor to discover steamable media source URLs.

#Synced-CL 51110937

diff --git a/Engine/Plugins/Media/MediaPlayerEditor/Source/MediaPlayerEditor/Private/Customizations/StreamMediaSourceCustomization.cpp b/Engine/Plugins/Media/MediaPlayerEditor/Source/MediaPlayerEditor/Private/Customizations/StreamMediaSourceCustomization.cpp
new file mode 100644
--- /dev/null
+++ b/Engine/Plugins/Media/MediaPlayerEditor/Source/MediaPlayerEditor/Private/Customizations/StreamMediaSourceCustomization.cpp
@@ -0,0 +1,151 @@
+﻿// Copyright Epic Games, Inc. All Rights Reserved.
+
+#include "StreamMediaSourceCustomization.h"
+
+#include "DetailCategoryBuilder.h"
+#include "DetailLayoutBuilder.h"
+#include "DetailWidgetRow.h"
+#include "IMediaCaptureSupport.h"
+#include "MediaCaptureSupport.h"
+#include "MediaPlayerEditorModule.h"
+#include "StreamMediaSource.h"
+#include "Framework/MultiBox/MultiBoxBuilder.h"
+#include "Modules/ModuleManager.h"
+#include "Widgets/Images/SImage.h"
+#include "Widgets/Input/SComboButton.h"
+
+#define LOCTEXT_NAMESPACE "FStreamMediaSourceCustomization"
+
+void FStreamMediaSourceCustomization::CustomizeDetails(IDetailLayoutBuilder& DetailBuilder)
+{
+	// Manual sort for the stream media source's property categories to ensure that 'Stream' category apppears on top
+	DetailBuilder.SortCategories([](const TMap<FName, IDetailCategoryBuilder*>& CategoryMap)
+	{
+		for (const TPair<FName, IDetailCategoryBuilder*>& Pair : CategoryMap)
+		{
+			int32 SortOrder = Pair.Value->GetSortOrder();
+			const FName CategoryName = Pair.Key;
+
+			if (CategoryName != "Stream")
+			{
+				SortOrder += 10;
+			}
+
+			Pair.Value->SetSortOrder(SortOrder);
+		}
+	});
+
+	StreamUrlHandle = DetailBuilder.GetProperty(GET_MEMBER_NAME_CHECKED(UStreamMediaSource, StreamUrl));
+	if (IDetailPropertyRow* Row = DetailBuilder.EditDefaultProperty(StreamUrlHandle))
+	{
+		IMediaPlayerEditorModule* MediaPlayerEditorModule = FModuleManager::LoadModulePtr<IMediaPlayerEditorModule>("MediaPlayerEditor");
+		TSharedPtr<ISlateStyle> Style = MediaPlayerEditorModule->GetStyle();
+		
+		TSharedPtr<SWidget> NameWidget;
+		TSharedPtr<SWidget> ValueWidget;
+		FDetailWidgetRow& CustomWidget = Row->CustomWidget();
+		
+		Row->GetDefaultWidgets(NameWidget, ValueWidget);
+		
+		CustomWidget
+		.NameContent()
+		[
+			NameWidget.ToSharedRef()
+		];
+
+		CustomWidget.ValueContent()
+		.MinDesiredWidth(125.0f)
+		.MaxDesiredWidth(600.0f)
+		[
+			SNew(SHorizontalBox)
+
+			+SHorizontalBox::Slot()
+			[
+				ValueWidget.ToSharedRef()
+			]
+
+			+SHorizontalBox::Slot()
+			.AutoWidth()
+			[
+				SNew(SComboButton)
+				.ComboButtonStyle(FAppStyle::Get(), "SimpleComboButtonWithIcon")
+				.OnGetMenuContent(this, &FStreamMediaSourceCustomization::GetDiscoverUrlMenu)
+				.ButtonContent()
+				[
+					SNew(SImage)
+					.Image(Style->GetBrush("MediaPlayerEditor.DiscoverUrl"))
+					.ColorAndOpacity(FSlateColor::UseForeground())
+				]
+			]
+		];
+	}
+}
+
+TSharedRef<SWidget> FStreamMediaSourceCustomization::GetDiscoverUrlMenu()
+{
+	constexpr bool bCloseMenuAfterSelection = true;
+	FMenuBuilder MenuBuilder(bCloseMenuAfterSelection, nullptr);
+
+	MenuBuilder.BeginSection("CaptureDevicesSection", LOCTEXT("CaptureDevicesSection", "Capture Devices"));
+	{
+		MenuBuilder.AddSubMenu(
+			LOCTEXT("AudioMenuLabel", "Audio"),
+			LOCTEXT("AudioMenuTooltip", "Available audio capture devices"),
+			FNewMenuDelegate::CreateSP(this, &FStreamMediaSourceCustomization::GetAudioCaptureDevicesSubmenu),
+			false,
+			FSlateIcon()
+		);
+
+		MenuBuilder.AddSubMenu(
+			LOCTEXT("VideoMenuLabel", "Video"),
+			LOCTEXT("VideoMenuTooltip", "Available video capture devices"),
+			FNewMenuDelegate::CreateSP(this, &FStreamMediaSourceCustomization::GetVideoCaptureDevicesSubmenu),
+			false,
+			FSlateIcon()
+		);
+	}
+	MenuBuilder.EndSection();
+	
+	return MenuBuilder.MakeWidget();
+}
+
+void FStreamMediaSourceCustomization::GetAudioCaptureDevicesSubmenu(FMenuBuilder& MenuBuilder)
+{
+	TArray<FMediaCaptureDeviceInfo> DeviceInfos;
+	MediaCaptureSupport::EnumerateAudioCaptureDevices(DeviceInfos);
+	MakeCaptureDeviceMenu(DeviceInfos, MenuBuilder);
+}
+
+void FStreamMediaSourceCustomization::GetVideoCaptureDevicesSubmenu(FMenuBuilder& MenuBuilder)
+{
+	TArray<FMediaCaptureDeviceInfo> DeviceInfos;
+	MediaCaptureSupport::EnumerateVideoCaptureDevices(DeviceInfos);
+	MakeCaptureDeviceMenu(DeviceInfos, MenuBuilder);
+}
+
+void FStreamMediaSourceCustomization::MakeCaptureDeviceMenu(TArray<FMediaCaptureDeviceInfo>& DeviceInfos, FMenuBuilder& MenuBuilder)
+{
+	for (const auto& DeviceInfo : DeviceInfos)
+	{
+		MenuBuilder.AddMenuEntry(
+			DeviceInfo.DisplayName,
+			FText::GetEmpty(),
+			FSlateIcon(),
+			FUIAction(FExecuteAction::CreateSP(this, &FStreamMediaSourceCustomization::SetStreamMediaSourceUrl, DeviceInfo.Url)),
+			NAME_None,
+			EUserInterfaceActionType::Button
+		);
+	}
+}
+
+void FStreamMediaSourceCustomization::SetStreamMediaSourceUrl(FString InUrl)
+{
+	if (!StreamUrlHandle.IsValid() || !StreamUrlHandle->IsValidHandle())
+	{
+		return;
+	}
+
+	StreamUrlHandle->SetValue(InUrl);
+}
+
+#undef LOCTEXT_NAMESPACE
diff --git a/Engine/Plugins/Media/MediaPlayerEditor/Source/MediaPlayerEditor/Private/Customizations/StreamMediaSourceCustomization.h b/Engine/Plugins/Media/MediaPlayerEditor/Source/MediaPlayerEditor/Private/Customizations/StreamMediaSourceCustomization.h
new file mode 100644
--- /dev/null
+++ b/Engine/Plugins/Media/MediaPlayerEditor/Source/MediaPlayerEditor/Private/Customizations/StreamMediaSourceCustomization.h
@@ -0,0 +1,47 @@
+﻿// Copyright Epic Games, Inc. All Rights Reserved.
+
+#pragma once
+
+#include "IDetailCustomization.h"
+#include "Framework/MultiBox/MultiBoxBuilder.h"
+
+class IPropertyHandle;
+struct FMediaCaptureDeviceInfo;
+class SWidget;
+
+/** Customization for stream media sources, primarily used to add a 'discover URL' dropdown to the stream URL property */
+class FStreamMediaSourceCustomization : public IDetailCustomization
+{
+public:
+	static TSharedRef<IDetailCustomization> MakeInstance()
+	{
+		return MakeShareable(new FStreamMediaSourceCustomization());
+	}
+
+	//~ IDetailCustomization interface
+	
+public:
+	virtual void CustomizeDetails(IDetailLayoutBuilder& DetailBuilder) override;
+	
+	//~ End IDetailCustomization interface
+
+private:
+	/** Creates the dropdown menu for the 'discover URL' combo button */
+	TSharedRef<SWidget> GetDiscoverUrlMenu();
+	
+	/** Fills the menu builder with options for found audio source URLs */
+	void GetAudioCaptureDevicesSubmenu(FMenuBuilder& MenuBuilder);
+	
+	/** Fills the menu builder with options for found video source URLs */
+	void GetVideoCaptureDevicesSubmenu(FMenuBuilder& MenuBuilder);
+
+	/** Fills the menu builder with entries for all provided device infos */
+	void MakeCaptureDeviceMenu(TArray<FMediaCaptureDeviceInfo>& DeviceInfos, FMenuBuilder& MenuBuilder);
+
+	/** Raised when a user selects a URL entry in the 'discover URLs' dropdown, sets the stream URL to the selected URL */
+	void SetStreamMediaSourceUrl(FString InUrl);
+
+private:
+	/** Cached property handle for UStreamMediaSource.StreamUrl property */
+	TSharedPtr<IPropertyHandle> StreamUrlHandle;
+};
diff --git a/Engine/Plugins/Media/MediaPlayerEditor/Source/MediaPlayerEditor/Private/MediaPlayerEditorModule.cpp b/Engine/Plugins/Media/MediaPlayerEditor/Source/MediaPlayerEditor/Private/MediaPlayerEditorModule.cpp
--- a/Engine/Plugins/Media/MediaPlayerEditor/Source/MediaPlayerEditor/Private/MediaPlayerEditorModule.cpp
+++ b/Engine/Plugins/Media/MediaPlayerEditor/Source/MediaPlayerEditor/Private/MediaPlayerEditorModule.cpp
@@ -19,6 +19,7 @@
 
 #include "BaseMediaSource.h"
 #include "FileMediaSource.h"
+#include "StreamMediaSource.h"
 #include "IMediaAssetsModule.h"
 #include "MediaPlayer.h"
 #include "MediaSoundComponent.h"
@@ -36,6 +37,7 @@
 
 #include "Customizations/BaseMediaSourceCustomization.h"
 #include "Customizations/FileMediaSourceCustomization.h"
+#include "Customizations/StreamMediaSourceCustomization.h"
 #include "Customizations/MediaTextureCustomization.h"
 #include "Customizations/PlatformMediaSourceCustomization.h"
 
@@ -106,6 +108,7 @@ public:
 	{
 		BaseMediaSourceName = UBaseMediaSource::StaticClass()->GetFName();
 		FileMediaSourceName = UFileMediaSource::StaticClass()->GetFName();
+		StreamMediaSourceName = UStreamMediaSource::StaticClass()->GetFName();
 		MediaSoundComponentName = UMediaSoundComponent::StaticClass()->GetFName();
 		MediaTextureName = UMediaTexture::StaticClass()->GetFName();
 		PlatformMediaSourceName = UPlatformMediaSource::StaticClass()->GetFName();
@@ -187,6 +190,7 @@ protected:
 		FPropertyEditorModule& PropertyModule = FModuleManager::LoadModuleChecked<FPropertyEditorModule>("PropertyEditor");
 		{
 			PropertyModule.RegisterCustomClassLayout(FileMediaSourceName, FOnGetDetailCustomizationInstance::CreateStatic(&FFileMediaSourceCustomization::MakeInstance));
+			PropertyModule.RegisterCustomClassLayout(StreamMediaSourceName, FOnGetDetailCustomizationInstance::CreateStatic(&FStreamMediaSourceCustomization::MakeInstance));
 			PropertyModule.RegisterCustomClassLayout(MediaTextureName, FOnGetDetailCustomizationInstance::CreateStatic(&FMediaTextureCustomization::MakeInstance));
 			PropertyModule.RegisterCustomClassLayout(PlatformMediaSourceName, FOnGetDetailCustomizationInstance::CreateStatic(&FPlatformMediaSourceCustomization::MakeInstance));
 			PropertyModule.RegisterCustomClassLayout(BaseMediaSourceName, FOnGetDetailCustomizationInstance::CreateStatic(&FBaseMediaSourceCustomization::MakeInstance));
@@ -200,6 +204,7 @@ protected:
 		{
 			PropertyModule.UnregisterCustomClassLayout(BaseMediaSourceName);
 			PropertyModule.UnregisterCustomClassLayout(FileMediaSourceName);
+			PropertyModule.UnregisterCustomClassLayout(StreamMediaSourceName);
 			PropertyModule.UnregisterCustomClassLayout(MediaTextureName);
 			PropertyModule.UnregisterCustomClassLayout(PlatformMediaSourceName);
 		}
@@ -408,6 +413,7 @@ private:
 	/** Class names. */
 	FName BaseMediaSourceName;
 	FName FileMediaSourceName;
+	FName StreamMediaSourceName;
 	FName MediaSoundComponentName;
 	FName MediaTextureName;
 	FName PlatformMediaSourceName;
diff --git a/Engine/Plugins/Media/MediaPlayerEditor/Source/MediaPlayerEditor/Private/Shared/MediaPlayerEditorStyle.cpp b/Engine/Plugins/Media/MediaPlayerEditor/Source/MediaPlayerEditor/Private/Shared/MediaPlayerEditorStyle.cpp
--- a/Engine/Plugins/Media/MediaPlayerEditor/Source/MediaPlayerEditor/Private/Shared/MediaPlayerEditorStyle.cpp
+++ b/Engine/Plugins/Media/MediaPlayerEditor/Source/MediaPlayerEditor/Private/Shared/MediaPlayerEditorStyle.cpp
@@ -45,6 +45,7 @@ FMediaPlayerEditorStyle::FMediaPlayerEditorStyle()
 	// misc
 	Set("MediaPlayerEditor.DragDropBorder", new BOX_BRUSH("border_dragdrop", 0.5f));
 	Set("MediaPlayerEditor.MediaSourceOpened", new IMAGE_BRUSH("mediasource_opened", Icon8x8));
+	Set("MediaPlayerEditor.DiscoverUrl", new CORE_IMAGE_BRUSH_SVG( "Starship/EditorViewport/globe", Icon16x16 ));
 
 	// tabs
 	Set("MediaPlayerEditor.Tabs.Info", new IMAGE_BRUSH("tab_info_16x", Icon16x16));
