[MediaPlayerEditor] Adding an option to display playback timer in either timecodes or frame numbers.

#Synced-CL 51094493

diff --git a/Engine/Plugins/Media/MediaPlayerEditor/Source/MediaPlayerEditor/Private/Shared/MediaPlayerEditorSettings.h b/Engine/Plugins/Media/MediaPlayerEditor/Source/MediaPlayerEditor/Private/Shared/MediaPlayerEditorSettings.h
--- a/Engine/Plugins/Media/MediaPlayerEditor/Source/MediaPlayerEditor/Private/Shared/MediaPlayerEditorSettings.h
+++ b/Engine/Plugins/Media/MediaPlayerEditor/Source/MediaPlayerEditor/Private/Shared/MediaPlayerEditorSettings.h
@@ -30,6 +30,16 @@ enum class EMediaPlayerEditorBackground : uint8
 	Checkered
 };
 
+/* Timer display units */
+UENUM()
+enum class EMediaPlayerEditorTimerUnit : uint8
+{
+	/** Display timer in HH:MM:SS format. */
+	Timecode,
+	/** Display timer as integer frame number. */
+	FrameNumber
+};
+
 UCLASS(config=EditorPerProjectUserSettings)
 class UMediaPlayerEditorSettings
 	: public UObject
@@ -46,6 +56,10 @@ public:
 	UPROPERTY(config, EditAnywhere, Category=Viewer)
 	bool ShowTextOverlays;
 
+	/** Whether to display the playback timer as frame number or timecode. */
+	UPROPERTY(config, EditAnywhere, Category=Viewer)
+	EMediaPlayerEditorTimerUnit TimerUnit = EMediaPlayerEditorTimerUnit::Timecode;
+
 	/** How the video viewport should be scaled. */
 	UPROPERTY(config, EditAnywhere, Category=Viewer)
 	EMediaPlayerEditorScale ViewportScale;
diff --git a/Engine/Plugins/Media/MediaPlayerEditor/Source/MediaPlayerEditor/Private/Widgets/SMediaPlayerEditorViewer.cpp b/Engine/Plugins/Media/MediaPlayerEditor/Source/MediaPlayerEditor/Private/Widgets/SMediaPlayerEditorViewer.cpp
--- a/Engine/Plugins/Media/MediaPlayerEditor/Source/MediaPlayerEditor/Private/Widgets/SMediaPlayerEditorViewer.cpp
+++ b/Engine/Plugins/Media/MediaPlayerEditor/Source/MediaPlayerEditor/Private/Widgets/SMediaPlayerEditorViewer.cpp
@@ -45,6 +45,17 @@
 
 #define LOCTEXT_NAMESPACE "SMediaPlayerEditorViewer"
 
+namespace UE::MediaPlayerEditor
+{
+	FTimespan GetVideoFrameDuration(const UMediaPlayer* InMediaPlayer)
+	{
+		// Get selected track's frame rate.
+		constexpr int32 SelectedTrackIndex = INDEX_NONE;
+		constexpr int32 SelectedFormatIndex = INDEX_NONE;
+		const float Framerate = InMediaPlayer ? InMediaPlayer->GetVideoTrackFrameRate(SelectedTrackIndex, SelectedFormatIndex) : 0.0f;
+		return Framerate > 0.0f ? FTimespan::FromSeconds(1.0f / Framerate) : FTimespan::Zero();
+	}
+}
 
 /* SMediaPlayerEditorPlayer structors
  *****************************************************************************/
@@ -239,13 +250,21 @@ void SMediaPlayerEditorViewer::Construct(const FArguments& InArgs, UMediaPlayer&
 				FSlateIcon()
 			);
 
+			SettingsMenuBuilder.AddSubMenu(
+				LOCTEXT("TimerUnitMenuLabel", "Timer Units"),
+				LOCTEXT("TimerUnitMenuTooltip", "Change the timer units"),
+				FNewMenuDelegate::CreateRaw(this, &SMediaPlayerEditorViewer::HandleTimerUnitMenuNewMenu),
+				false,
+				FSlateIcon()
+			);
+
 			SettingsMenuBuilder.AddMenuEntry(
 				LOCTEXT("ShowTextOverlaysMenuLabel", "Show Text Overlays"),
 				LOCTEXT("ShowTextOverlaysMenuTooltip", "Show caption and subtitle text overlays"),
 				FSlateIcon(),
 				FUIAction(
 					FExecuteAction::CreateLambda([]{
-						auto Settings = GetMutableDefault<UMediaPlayerEditorSettings>();
+						UMediaPlayerEditorSettings* Settings = GetMutableDefault<UMediaPlayerEditorSettings>();
 						Settings->ShowTextOverlays = !Settings->ShowTextOverlays;
 						Settings->SaveConfig();
 					}),
@@ -654,7 +673,7 @@ FReply SMediaPlayerEditorViewer::OnDrop(const FGeometry& MyGeometry, const FDrag
 
 void SMediaPlayerEditorViewer::MakeCaptureDeviceMenu(TArray<FMediaCaptureDeviceInfo>& DeviceInfos, FMenuBuilder& MenuBuilder)
 {
-	for (const auto& DeviceInfo : DeviceInfos)
+	for (const FMediaCaptureDeviceInfo& DeviceInfo : DeviceInfos)
 	{
 		MenuBuilder.AddMenuEntry(
 			DeviceInfo.DisplayName,
@@ -700,7 +719,7 @@ void SMediaPlayerEditorViewer::SetDesiredPlayerName(FName PlayerName)
 		}
 	}
 
-	auto Settings = GetMutableDefault<UMediaPlayerEditorSettings>();
+	UMediaPlayerEditorSettings* Settings = GetMutableDefault<UMediaPlayerEditorSettings>();
 	{
 		Settings->DesiredPlayerName = PlayerName;
 		Settings->SaveConfig();
@@ -738,7 +757,7 @@ void SMediaPlayerEditorViewer::HandleDecoderMenuNewMenu(FMenuBuilder& MenuBuilde
 	MenuBuilder.AddMenuSeparator();
 
 	// get registered player plug-ins
-	auto MediaModule = FModuleManager::LoadModulePtr<IMediaModule>("Media");
+	const IMediaModule* MediaModule = FModuleManager::LoadModulePtr<IMediaModule>("Media");
 
 	if (MediaModule == nullptr)
 	{
@@ -924,7 +943,7 @@ void SMediaPlayerEditorViewer::HandleScaleMenuNewMenu(FMenuBuilder& MenuBuilder)
 		FSlateIcon(),
 		FUIAction(
 			FExecuteAction::CreateLambda([]{
-				auto Settings = GetMutableDefault<UMediaPlayerEditorSettings>();
+				UMediaPlayerEditorSettings* Settings = GetMutableDefault<UMediaPlayerEditorSettings>();
 				Settings->ViewportScale = EMediaPlayerEditorScale::Fit;
 				Settings->SaveConfig();
 			}),
@@ -941,7 +960,7 @@ void SMediaPlayerEditorViewer::HandleScaleMenuNewMenu(FMenuBuilder& MenuBuilder)
 		FSlateIcon(),
 		FUIAction(
 			FExecuteAction::CreateLambda([]{
-				auto Settings = GetMutableDefault<UMediaPlayerEditorSettings>();
+				UMediaPlayerEditorSettings* Settings = GetMutableDefault<UMediaPlayerEditorSettings>();
 				Settings->ViewportScale = EMediaPlayerEditorScale::Fill;
 				Settings->SaveConfig();
 			}),
@@ -958,7 +977,7 @@ void SMediaPlayerEditorViewer::HandleScaleMenuNewMenu(FMenuBuilder& MenuBuilder)
 		FSlateIcon(),
 		FUIAction(
 			FExecuteAction::CreateLambda([]{
-				auto Settings = GetMutableDefault<UMediaPlayerEditorSettings>();
+				UMediaPlayerEditorSettings* Settings = GetMutableDefault<UMediaPlayerEditorSettings>();
 				Settings->ViewportScale = EMediaPlayerEditorScale::Original;
 				Settings->SaveConfig();
 			}),
@@ -970,6 +989,42 @@ void SMediaPlayerEditorViewer::HandleScaleMenuNewMenu(FMenuBuilder& MenuBuilder)
 	);
 }
 
+void SMediaPlayerEditorViewer::HandleTimerUnitMenuNewMenu(FMenuBuilder& MenuBuilder)
+{
+	MenuBuilder.AddMenuEntry(
+		LOCTEXT("TimerUnitTimecodesMenuLabel", "Timecode (HH:MM:SS)"),
+		LOCTEXT("TimerUnitTimecodesMenuTooltip", "Display the timer as HH:MM:SS formatted timecode."),
+		FSlateIcon(),
+		FUIAction(
+			FExecuteAction::CreateLambda([]{
+				UMediaPlayerEditorSettings* Settings = GetMutableDefault<UMediaPlayerEditorSettings>();
+				Settings->TimerUnit = EMediaPlayerEditorTimerUnit::Timecode;
+				Settings->SaveConfig();
+			}),
+			FCanExecuteAction(),
+			FIsActionChecked::CreateLambda([]{ return (GetDefault<UMediaPlayerEditorSettings>()->TimerUnit == EMediaPlayerEditorTimerUnit::Timecode); })
+		),
+		NAME_None,
+		EUserInterfaceActionType::RadioButton
+	);
+
+	MenuBuilder.AddMenuEntry(
+		LOCTEXT("TimerUnitFramesMenuLabel", "Frame Number"),
+		LOCTEXT("TimerUnitFramesMenuTooltip", "Display the timer as frame number. Note: will revert to \"timecode\" for variable frame rate videos."),
+		FSlateIcon(),
+		FUIAction(
+			FExecuteAction::CreateLambda([]{
+				UMediaPlayerEditorSettings* Settings = GetMutableDefault<UMediaPlayerEditorSettings>();
+				Settings->TimerUnit = EMediaPlayerEditorTimerUnit::FrameNumber;
+				Settings->SaveConfig();
+			}),
+			FCanExecuteAction(),
+			FIsActionChecked::CreateLambda([]{ return (GetDefault<UMediaPlayerEditorSettings>()->TimerUnit == EMediaPlayerEditorTimerUnit::FrameNumber); })
+		),
+		NAME_None,
+		EUserInterfaceActionType::RadioButton
+	);
+}
 
 FText SMediaPlayerEditorViewer::HandleTimerTextBlockText() const
 {
@@ -978,15 +1033,41 @@ FText SMediaPlayerEditorViewer::HandleTimerTextBlockText() const
 		return FText::GetEmpty();
 	}
 
-	FTimespan Time = MediaPlayer->GetDisplayTime();
+	const FTimespan Time = MediaPlayer->GetDisplayTime();
 
 	// empty string if time n/a
 	if (Time < FTimespan::Zero())
 	{
 		return FText::GetEmpty();
 	}
+	
+	const FTimespan Duration = MediaPlayer->GetDuration();
 
-	FTimespan Duration = MediaPlayer->GetDuration();
+	const FText TimerTextBlockFormat = LOCTEXT("TimerTextBlockFormat", "{0} / {1}");
+
+	if (GetDefault<UMediaPlayerEditorSettings>()->TimerUnit == EMediaPlayerEditorTimerUnit::FrameNumber)
+	{
+		const FTimespan FrameDuration = UE::MediaPlayerEditor::GetVideoFrameDuration(MediaPlayer);
+
+		// Note: if the frame rate is not available, we fall back to displaying the time code. 
+		if (FrameDuration != FTimespan::Zero())
+		{
+			const int32 FrameNumber = static_cast<int32>(FMath::RoundHalfToEven(Time.GetTotalSeconds() / FrameDuration.GetTotalSeconds()));
+			
+			// time only if duration n/a
+			if (Duration <= FTimespan::Zero())
+			{
+				return FText::AsNumber(FrameNumber);
+			}
+
+			// format time & duration
+			const FText DurationText = (Duration == FTimespan::MaxValue())
+				? FText::FromString(TEXT("\u221E")) // infinity symbol
+				: FText::AsNumber(static_cast<int32>(FMath::RoundHalfToEven(Duration.GetTotalSeconds() / FrameDuration.GetTotalSeconds())));
+
+			return FText::Format(TimerTextBlockFormat, FText::AsNumber(FrameNumber), DurationText);
+		}
+	}
 
 	// time only if duration n/a
 	if (Duration <= FTimespan::Zero())
@@ -999,7 +1080,7 @@ FText SMediaPlayerEditorViewer::HandleTimerTextBlockText() const
 		? FText::FromString(TEXT("\u221E")) // infinity symbol
 		: FText::AsTimespan(Duration);
 
-	return FText::Format(LOCTEXT("TimerTextBlockFormat", "{0} / {1}"), FText::AsTimespan(Time), DurationText);
+	return FText::Format(TimerTextBlockFormat, FText::AsTimespan(Time), DurationText);
 }
 
 
diff --git a/Engine/Plugins/Media/MediaPlayerEditor/Source/MediaPlayerEditor/Public/Widgets/SMediaPlayerEditorViewer.h b/Engine/Plugins/Media/MediaPlayerEditor/Source/MediaPlayerEditor/Public/Widgets/SMediaPlayerEditorViewer.h
--- a/Engine/Plugins/Media/MediaPlayerEditor/Source/MediaPlayerEditor/Public/Widgets/SMediaPlayerEditorViewer.h
+++ b/Engine/Plugins/Media/MediaPlayerEditor/Source/MediaPlayerEditor/Public/Widgets/SMediaPlayerEditorViewer.h
@@ -153,6 +153,9 @@ private:
 	/** Callback for creating the Background sub-menu. */
 	UE_API void HandleBackgroundMenuNewMenu(FMenuBuilder& MenuBuilder);
 
+	/** Callback for creating the timer units sub-menu. */
+	UE_API void HandleTimerUnitMenuNewMenu(FMenuBuilder& MenuBuilder);
+	
 	/** Callback for creating a track sub-menu. */
 	UE_API void HandleTrackMenuNewMenu(FMenuBuilder& MenuBuilder, EMediaPlayerTrack TrackType);
 
